UnrealToolbox: Add an spinner while a tool is updating.

#rnx

diff --git a/Engine/Source/Programs/UnrealToolbox/ToolCatalogViewModel.cs b/Engine/Source/Programs/UnrealToolbox/ToolCatalogViewModel.cs
--- a/Engine/Source/Programs/UnrealToolbox/ToolCatalogViewModel.cs
+++ b/Engine/Source/Programs/UnrealToolbox/ToolCatalogViewModel.cs
@@ -121,6 +121,9 @@ namespace UnrealToolbox
 		[ObservableProperty]
 		bool _showRemove;
 
+		[ObservableProperty]
+		bool _isBusy;
+
 		internal void UpdateFromModel()
 		{
 			Name = _item.Name;
@@ -141,6 +144,7 @@ namespace UnrealToolbox
 				{
 					Status = $"Updating to {pending.Deployment.Version}...";
 				}
+				IsBusy = true;
 			}
 			else
 			{
@@ -153,6 +157,7 @@ namespace UnrealToolbox
 				{
 					Status = String.Empty;
 				}
+				IsBusy = false;
 			}
 
 			ShowAdd = _item.Latest != null && _item.Current == null;
