UnrealToolbox: Update config file name for Unreal Toolbox to "Toolbox.json".

#rnx

diff --git a/Engine/Source/Programs/UnrealToolbox/ToolCatalog.cs b/Engine/Source/Programs/UnrealToolbox/ToolCatalog.cs
--- a/Engine/Source/Programs/UnrealToolbox/ToolCatalog.cs
+++ b/Engine/Source/Programs/UnrealToolbox/ToolCatalog.cs
@@ -347,7 +347,7 @@ namespace UnrealToolbox
 
 		ToolConfig LoadToolConfig(DirectoryReference toolDir)
 		{
-			FileReference configFile = FileReference.Combine(toolDir, "Launcher.json");
+			FileReference configFile = FileReference.Combine(toolDir, "Toolbox.json");
 			if (FileReference.Exists(configFile))
 			{
 				JsonSerializerOptions jsonSerializerOptions = new JsonSerializerOptions();
