Unreal Toolbox: Run install commands in the working directory of the downloaded tool.

#rnx

diff --git a/Engine/Source/Programs/UnrealToolbox/ToolCatalog.cs b/Engine/Source/Programs/UnrealToolbox/ToolCatalog.cs
--- a/Engine/Source/Programs/UnrealToolbox/ToolCatalog.cs
+++ b/Engine/Source/Programs/UnrealToolbox/ToolCatalog.cs
@@ -494,10 +494,13 @@ namespace UnrealToolbox
 
 			if (toolConfig.InstallCommand != null)
 			{
+				PendingToolDeploymentInfo pending = new PendingToolDeploymentInfo(false, "Running install actions...", deploymentInfo);
+				UpdateStateAndNotify(item, state => state with { Pending = pending });
+
 				int exitCode = await RunCommandAsync(item.Id.ToString(), toolConfig.InstallCommand, toolDir.FullName, cancellationToken);
 				if (exitCode != 0)
 				{
-					PendingToolDeploymentInfo pending = new PendingToolDeploymentInfo(true, $"Installation failed ({exitCode}). See log.", deploymentInfo);
+					pending = new PendingToolDeploymentInfo(true, $"Installation failed ({exitCode}). See log.", deploymentInfo);
 					UpdateStateAndNotify(item, state => state with { Pending = pending });
 					return;
 				}
@@ -543,8 +546,14 @@ namespace UnrealToolbox
 				arguments = CommandLineArguments.Join(command.Arguments);
 			}
 
+			string fileName = Path.Combine(workingDir, command.FileName);
+			if (!File.Exists(fileName))
+			{
+				fileName = command.FileName;
+			}
+
 			using (ManagedProcessGroup newProcessGroup = new ManagedProcessGroup())
-			using (ManagedProcess newProcess = new ManagedProcess(newProcessGroup, command.FileName, arguments, workingDir, null, null, ProcessPriorityClass.Normal))
+			using (ManagedProcess newProcess = new ManagedProcess(newProcessGroup, fileName, arguments, workingDir, null, null, ProcessPriorityClass.Normal))
 			{
 				for (; ; )
 				{
