UnrealToolbox: Add a flag indicating that tools have a manual install process and should not be initiated automatically.

#rnx

diff --git a/Engine/Source/Programs/UnrealToolbox/ToolCatalog.cs b/Engine/Source/Programs/UnrealToolbox/ToolCatalog.cs
--- a/Engine/Source/Programs/UnrealToolbox/ToolCatalog.cs
+++ b/Engine/Source/Programs/UnrealToolbox/ToolCatalog.cs
@@ -804,9 +804,12 @@ namespace UnrealToolbox
 				{
 					foreach (Item item in _items.Values)
 					{
-						if (item._state.Current != null && item._state.Latest != null && item._state.Latest.Id != item._state.Current.Id)
+						if (item._state.Current != null && !item._state.Current.Config.ManualInstall)
 						{
-							UpdateDeployment(item, item._state.Latest);
+							if (item._state.Latest != null && item._state.Latest.Id != item._state.Current.Id)
+							{
+								UpdateDeployment(item, item._state.Latest);
+							}
 						}
 					}
 				}
diff --git a/Engine/Source/Programs/UnrealToolbox/ToolConfig.cs b/Engine/Source/Programs/UnrealToolbox/ToolConfig.cs
--- a/Engine/Source/Programs/UnrealToolbox/ToolConfig.cs
+++ b/Engine/Source/Programs/UnrealToolbox/ToolConfig.cs
@@ -12,6 +12,11 @@ namespace UnrealToolbox
 		/// </summary>
 		public ToolMenuItem? PopupMenu { get; set; }
 
+		/// <summary>
+		/// Whether the install/uninstall for this tool needs to be triggered manually
+		/// </summary>
+		public bool ManualInstall { get; set; }
+
 		/// <summary>
 		/// Command to run when installing
 		/// </summary>
