UnrealToolbox: Fix deletion of tool directories containing read-only files.

#rnx

diff --git a/Engine/Source/Programs/Shared/EpicGames.Core/FileUtils.cs b/Engine/Source/Programs/Shared/EpicGames.Core/FileUtils.cs
--- a/Engine/Source/Programs/Shared/EpicGames.Core/FileUtils.cs
+++ b/Engine/Source/Programs/Shared/EpicGames.Core/FileUtils.cs
@@ -543,15 +543,19 @@ namespace EpicGames.Core
 			{
 				dirName = RawPathPrefix + dirName;
 			}
+			if (!dirName.EndsWith("\\"))
+			{
+				dirName += "\\";
+			}
 
-			IntPtr hFind = FindFirstFileW(dirName + "\\*", ref findData);
+			IntPtr hFind = FindFirstFileW(dirName + "*", ref findData);
 			if (hFind != INVALID_HANDLE_VALUE)
 			{
 				try
 				{
 					for (; ; )
 					{
-						string fullName = dirName + "\\" + findData.cFileName;
+						string fullName = dirName + findData.cFileName;
 						if ((findData.dwFileAttributes & FILE_ATTRIBUTE_DIRECTORY) != 0)
 						{
 							if (findData.cFileName != "." && findData.cFileName != "..")
diff --git a/Engine/Source/Programs/UnrealToolbox/ToolCatalog.cs b/Engine/Source/Programs/UnrealToolbox/ToolCatalog.cs
--- a/Engine/Source/Programs/UnrealToolbox/ToolCatalog.cs
+++ b/Engine/Source/Programs/UnrealToolbox/ToolCatalog.cs
@@ -685,13 +685,13 @@ namespace UnrealToolbox
 			// Try to delete all the deployment directories
 			foreach (DirectoryInfo deleteDeploymentDir in deleteDeploymentDirs)
 			{
-				result &= TryDeleteDir(deleteDeploymentDir, true);
+				result &= ForceDeleteDir(deleteDeploymentDir);
 			}
 
 			// Try to delete all the tool directories if they're empty
 			foreach (DirectoryInfo deleteToolDir in deleteToolDirs)
 			{
-				result &= TryDeleteDir(deleteToolDir, false);
+				result &= TryDeleteDir(deleteToolDir);
 			}
 
 			return result;
@@ -711,11 +711,25 @@ namespace UnrealToolbox
 			}
 		}
 
-		bool TryDeleteDir(DirectoryInfo dirInfo, bool recursive)
+		bool ForceDeleteDir(DirectoryInfo dirInfo)
 		{
 			try
 			{
-				dirInfo.Delete(recursive);
+				FileUtils.ForceDeleteDirectory(dirInfo);
+				return true;
+			}
+			catch (Exception ex)
+			{
+				_logger.LogWarning(ex, "Unable to delete directory {DirName}: {Message}", dirInfo.FullName, ex.Message);
+				return false;
+			}
+		}
+
+		bool TryDeleteDir(DirectoryInfo dirInfo)
+		{
+			try
+			{
+				dirInfo.Delete(false);
 				return true;
 			}
 			catch (Exception ex)
