UnrealToolbox: Add a self-update mechanism using the unreal-toolbox tool id in Horde.

diff --git a/Engine/Source/Programs/Shared/EpicGames.Core/DirectoryReference.cs b/Engine/Source/Programs/Shared/EpicGames.Core/DirectoryReference.cs
--- a/Engine/Source/Programs/Shared/EpicGames.Core/DirectoryReference.cs
+++ b/Engine/Source/Programs/Shared/EpicGames.Core/DirectoryReference.cs
@@ -222,6 +222,14 @@ namespace EpicGames.Core
 		/// <inheritdoc/>
 		public static bool operator >=(DirectoryReference left, DirectoryReference right) => left is null ? right is null : left.CompareTo(right) >= 0;
 
+		/// <summary>
+		/// Moves a directory from one location to another
+		/// </summary>
+		/// <param name="sourceDir">The source directory</param>
+		/// <param name="destDir">The destination directory</param>
+		public static void Move(DirectoryReference sourceDir, DirectoryReference destDir)
+			=> Directory.Move(sourceDir.FullName, destDir.FullName);
+
 		/// <summary>
 		/// Helper function to create a remote directory reference. Unlike normal DirectoryReference objects, these aren't converted to a full path in the local filesystem.
 		/// </summary>
diff --git a/Engine/Source/Programs/UnrealToolbox/App.axaml.cs b/Engine/Source/Programs/UnrealToolbox/App.axaml.cs
--- a/Engine/Source/Programs/UnrealToolbox/App.axaml.cs
+++ b/Engine/Source/Programs/UnrealToolbox/App.axaml.cs
@@ -38,6 +38,7 @@ namespace UnrealToolbox
 			serviceCollection.AddSingleton<ToolCatalog>();
 			serviceCollection.AddSingleton<IToolCatalog, ToolCatalog>(sp => sp.GetRequiredService<ToolCatalog>());
 			serviceCollection.AddSingleton<ITrayAppHost>(this);
+			serviceCollection.AddSingleton<SelfUpdateService>();
 
 			// Configure plugins
 			serviceCollection.AddSingleton<ITrayAppPlugin, HordeAgentPlugin>();
@@ -61,6 +62,23 @@ namespace UnrealToolbox
 			ToolCatalog toolCatalog = _serviceProvider.GetRequiredService<ToolCatalog>();
 			toolCatalog.OnItemsChanged += UpdateMenu;
 			toolCatalog.Start();
+
+			SelfUpdateService selfUpdateService = _serviceProvider.GetRequiredService<SelfUpdateService>();
+			selfUpdateService.OnUpdateReady += UpdateReady;
+			selfUpdateService.Start();
+		}
+
+		private void UpdateReady()
+		{
+			Dispatcher.UIThread.Post(() => UpdateReadyMainThread());
+		}
+
+		void UpdateReadyMainThread()
+		{
+			if (_settingsWindow == null)
+			{
+				((IClassicDesktopStyleApplicationLifetime)ApplicationLifetime!).Shutdown();
+			}
 		}
 
 		private void UpdateMenu()
diff --git a/Engine/Source/Programs/UnrealToolbox/Program.cs b/Engine/Source/Programs/UnrealToolbox/Program.cs
--- a/Engine/Source/Programs/UnrealToolbox/Program.cs
+++ b/Engine/Source/Programs/UnrealToolbox/Program.cs
@@ -7,8 +7,6 @@ using Avalonia;
 using Avalonia.Controls;
 using EpicGames.Core;
 
-#pragma warning disable CA1806 // Not checking return code from MessageBox
-
 namespace UnrealToolbox
 {
 	static class Program
@@ -16,77 +14,43 @@ namespace UnrealToolbox
 		const string MutexName = "UnrealToolbox-Mutex";
 		const string EventName = "UnrealToolbox-Exit";
 
-		[DllImport("user32.dll", SetLastError = true, CharSet = CharSet.Auto)]
-		public static extern int MessageBox(IntPtr hWnd, string text, string? caption, uint type);
+		public static SelfUpdateState? Update { get; private set; }
 
 		[STAThread]
 		public static int Main(string[] args)
 		{
-			using EventWaitHandle closeEvent = new EventWaitHandle(false, EventResetMode.AutoReset, EventName);
-			using SingleInstanceMutex mutex = new SingleInstanceMutex(MutexName);
-
-			if (args.Any(x => x.Equals("-close", StringComparison.OrdinalIgnoreCase)))
-			{
-				closeEvent.Set();
-				return mutex.Wait(5000) ? 0 : 1;
-			}
-
-			if (!mutex.Wait(0))
+			if (!args.Any(x => x.Equals("-NoUpdate", StringComparison.OrdinalIgnoreCase)))
 			{
-				return 1;
+				Update = SelfUpdateState.TryCreate("Unreal Toolbox", args);
 			}
 
-			if (args.Any(x => x.Equals("-shadowcopy", StringComparison.OrdinalIgnoreCase)))
+			for (; ; )
 			{
-				DirectoryReference? localAppData = DirectoryReference.GetSpecialFolder(Environment.SpecialFolder.LocalApplicationData);
-				if (localAppData != null)
+				if (Update != null && Update.TryLaunchLatest())
 				{
-					FileReference sourceExe = new FileReference(Assembly.GetExecutingAssembly().Location).ChangeExtension(".exe");
-					if (FileReference.Exists(sourceExe) && !sourceExe.IsUnderDirectory(localAppData))
-					{
-						DirectoryReference sourceDir = sourceExe.Directory;
-						DirectoryReference targetDir = DirectoryReference.Combine(localAppData, "Epic Games", "UnrealToolbox");
-						try
-						{
-							DirectoryReference.CreateDirectory(targetDir);
-							FileUtils.ForceDeleteDirectoryContents(targetDir);
-
-							CopyFiles(sourceDir, targetDir);
-							mutex.Release();
-
-							FileReference targetExe = FileReference.Combine(targetDir, sourceExe.MakeRelativeTo(sourceDir)).ChangeExtension(".exe");
-							using Process process = Process.Start(targetExe.FullName, args);
+					return 0;
+				}
 
-							return 0;
-						}
-						catch (Exception ex)
-						{
-							MessageBox(IntPtr.Zero, $"Unable to copy app to temp location. Error:\n\n{ex}", null, 0);
-							return 1;
-						}
-					}
+				int result = RealMain(args);
+				if (Update == null || !Update.IsUpdatePending())
+				{
+					return result;
 				}
 			}
-
-			BuildAvaloniaApp()
-				.StartWithClassicDesktopLifetime(args, ShutdownMode.OnExplicitShutdown);
-
-			return 0;
 		}
 
-		static void CopyFiles(DirectoryReference sourceDir, DirectoryReference targetDir)
+		static int RealMain(string[] args)
 		{
-			DirectoryReference.CreateDirectory(targetDir);
-			foreach (DirectoryReference sourceSubDir in DirectoryReference.EnumerateDirectories(sourceDir))
-			{
-				DirectoryReference targetSubDir = DirectoryReference.Combine(targetDir, sourceSubDir.GetDirectoryName());
-				CopyFiles(sourceSubDir, targetSubDir);
-			}
-			foreach (FileReference sourceFile in DirectoryReference.EnumerateFiles(sourceDir))
+			using SingleInstanceMutex mutex = new SingleInstanceMutex(MutexName);
+			if (!mutex.Wait(0))
 			{
-				FileReference targetFile = FileReference.Combine(targetDir, sourceFile.GetFileName());
-				FileReference.Copy(sourceFile, targetFile, true);
+				return 1;
 			}
+
+			BuildAvaloniaApp()
+				.StartWithClassicDesktopLifetime(args, ShutdownMode.OnExplicitShutdown);
+
+			return 0;
 		}
 
 		// Avalonia configuration, don't remove; also used by visual designer.
diff --git a/Engine/Source/Programs/UnrealToolbox/Properties/launchSettings.json b/Engine/Source/Programs/UnrealToolbox/Properties/launchSettings.json
new file mode 100644
--- /dev/null
+++ b/Engine/Source/Programs/UnrealToolbox/Properties/launchSettings.json
@@ -0,0 +1,8 @@
+{
+  "profiles": {
+    "UnrealToolbox": {
+      "commandName": "Project",
+      "commandLineArgs": "-NoUpdate"
+    }
+  }
+}
\ No newline at end of file
