UnrealToolbox: Missing file.

#rnx

diff --git a/Engine/Source/Programs/UnrealToolbox/SelfUpdateService.cs b/Engine/Source/Programs/UnrealToolbox/SelfUpdateService.cs
new file mode 100644
--- /dev/null
+++ b/Engine/Source/Programs/UnrealToolbox/SelfUpdateService.cs
@@ -0,0 +1,241 @@
+using EpicGames.Core;
+using EpicGames.Horde.Storage.Nodes;
+using EpicGames.Horde.Tools;
+using Microsoft.AspNetCore.StaticFiles;
+using Microsoft.Extensions.Logging;
+using System.Diagnostics;
+using System.Reflection;
+
+namespace UnrealToolbox
+{
+	class SelfUpdateState
+	{
+		class JsonState
+		{
+			public string? LaunchApp { get; set; }
+			public string? LatestVersion { get; set; }
+			public string? UpdateVersion { get; set; }
+		}
+
+		readonly DirectoryReference _baseDir;
+		readonly JsonConfig<JsonState> _config;
+		readonly string[] _args;
+		readonly bool _isLaunchApp;
+
+		public string? LatestVersion
+			=> _config.Current.LatestVersion;
+
+		public DirectoryReference LatestDir
+			=> DirectoryReference.Combine(_baseDir, "Latest");
+
+		public DirectoryReference UpdateDir
+			=> DirectoryReference.Combine(_baseDir, "Update");
+
+		public string? UpdateVersion
+		{
+			get => _config.Current.UpdateVersion;
+			set => UpdateState(x => x.UpdateVersion = value);
+		}
+
+		public SelfUpdateState(DirectoryReference baseDir, string[] args)
+		{
+			_baseDir = baseDir;
+			_config = new JsonConfig<JsonState>(FileReference.Combine(baseDir, "Update.json"));
+			_config.LoadSettings();
+			_args = args;
+
+			FileReference currentAssembly = new FileReference(Assembly.GetExecutingAssembly().Location);
+
+			_isLaunchApp = !currentAssembly.IsUnderDirectory(_baseDir);
+			if(_isLaunchApp && !String.Equals(_config.Current.LaunchApp, currentAssembly.FullName, StringComparison.OrdinalIgnoreCase))
+			{
+				UpdateState(x => x.LaunchApp = currentAssembly.FullName);
+			}
+		}
+
+		public bool IsUpdatePending()
+			=> !String.IsNullOrEmpty(_config.Current.UpdateVersion);
+
+		void UpdateState(Action<JsonState> update)
+		{
+			JsonState curState = _config.Current;
+
+			JsonState newState = new JsonState();
+			newState.LaunchApp = curState.LaunchApp;
+			newState.LatestVersion = curState.LatestVersion;
+			newState.UpdateVersion = curState.UpdateVersion;
+			update(newState);
+
+			_config.UpdateSettings(newState);
+		}
+
+		public static SelfUpdateState? TryCreate(string appName, string[] args)
+		{
+			DirectoryReference? localAppData = DirectoryReference.GetSpecialFolder(Environment.SpecialFolder.LocalApplicationData);
+			if (localAppData == null)
+			{
+				return null;
+			}
+
+			DirectoryReference baseDir = DirectoryReference.Combine(localAppData, "Epic Games", appName);
+			return new SelfUpdateState(baseDir, args);
+		}
+
+		public bool TryLaunchLatest()
+		{
+			if (_isLaunchApp)
+			{
+				string launchApp = Assembly.GetExecutingAssembly().Location;
+				if (!String.Equals(_config.Current.LaunchApp, launchApp, StringComparison.OrdinalIgnoreCase))
+				{
+					UpdateState(x => x.LaunchApp = launchApp);
+				}
+
+				string? updateVersion = _config.Current.UpdateVersion;
+				if (!String.IsNullOrEmpty(updateVersion))
+				{
+					if (!String.IsNullOrEmpty(_config.Current.LatestVersion))
+					{
+						UpdateState(x => x.LatestVersion = null);
+						FileUtils.ForceDeleteDirectory(LatestDir);
+					}
+
+					UpdateState(x => x.UpdateVersion = null);
+					DirectoryReference.Move(UpdateDir, LatestDir);
+					UpdateState(x => x.LatestVersion = updateVersion);
+				}
+
+				if (!String.IsNullOrEmpty(_config.Current.LatestVersion))
+				{
+					string fileName = Path.GetFileName(Assembly.GetExecutingAssembly().Location);
+					FileReference executable = FileReference.Combine(LatestDir, fileName);
+					if (FileReference.Exists(executable))
+					{
+						return Launch(executable.FullName, _args);
+					}
+				}
+			}
+			else
+			{
+				if (!String.IsNullOrEmpty(_config.Current.LaunchApp))
+				{
+					return Launch(_config.Current.LaunchApp, _args);
+				}
+			}
+			return false;
+		}
+
+		static bool Launch(string executable, string[] args)
+		{
+			using Process process = new Process();
+			if (executable.EndsWith(".dll", StringComparison.OrdinalIgnoreCase))
+			{
+				process.StartInfo.FileName = "dotnet";
+				process.StartInfo.ArgumentList.Add(executable);
+			}
+			else
+			{
+				process.StartInfo.FileName = executable;
+			}
+
+			foreach (string arg in args)
+			{
+				process.StartInfo.ArgumentList.Add(arg);
+			}
+
+			return process.Start();
+		}
+	}
+
+	class SelfUpdateService : IAsyncDisposable
+	{
+		static ToolId ToolId { get; } = new ToolId("unreal-toolbox");
+
+		readonly IHordeClientProvider _hordeClientProvider;
+		readonly BackgroundTask _backgroundTask;
+		readonly ILogger _logger;
+
+		public event Action? OnUpdateReady;
+
+		public SelfUpdateService(IHordeClientProvider hordeClientProvider, ILogger<SelfUpdateService> logger)
+		{
+			_hordeClientProvider = hordeClientProvider;
+			_backgroundTask = new BackgroundTask(CheckForUpdatesAsync);
+			_logger = logger;
+		}
+
+		public void Start()
+		{
+			_backgroundTask.Start();
+		}
+
+		public async ValueTask DisposeAsync()
+		{
+			await _backgroundTask.DisposeAsync();
+		}
+
+		async Task CheckForUpdatesAsync(CancellationToken cancellationToken)
+		{
+			for (; ; )
+			{
+				try
+				{
+					await CheckForUpdate(cancellationToken);
+				}
+				catch (OperationCanceledException)
+				{
+					throw;
+				}
+				catch (Exception ex)
+				{
+					_logger.LogError(ex, "Error checking for updates: {Message}", ex.Message);
+				}
+				await Task.Delay(TimeSpan.FromMinutes(10.0), cancellationToken);
+			}
+		}
+
+		async Task CheckForUpdate(CancellationToken cancellationToken)
+		{
+			SelfUpdateState? selfUpdate = Program.Update;
+			if (selfUpdate == null)
+			{
+				return;
+			}
+
+			IHordeClientRef? clientRef = await _hordeClientProvider.GetClientRefAsync().WaitAsync(cancellationToken);
+			if (clientRef == null)
+			{
+				return;
+			}
+
+			ITool? tool = await clientRef.Client.Tools.GetAsync(ToolId, cancellationToken);
+			if (tool == null || tool.Deployments.Count == 0)
+			{
+				return;
+			}
+
+			IToolDeployment deployment = tool.Deployments[^1];
+			_logger.LogInformation("Latest app version is {Id}", deployment.Id);
+
+			string updateVersion = deployment.Id.ToString();
+			if (String.Equals(updateVersion, selfUpdate.LatestVersion, StringComparison.OrdinalIgnoreCase) || String.Equals(updateVersion, selfUpdate.UpdateVersion, StringComparison.OrdinalIgnoreCase))
+			{
+				selfUpdate.UpdateVersion = null;
+
+				DirectoryReference updateDir = selfUpdate.UpdateDir;
+				DirectoryReference.CreateDirectory(updateDir);
+
+				FileUtils.ForceDeleteDirectoryContents(updateDir);
+
+				await deployment.Content.ExtractAsync(updateDir.ToDirectoryInfo(), _logger, cancellationToken);
+
+				selfUpdate.UpdateVersion = updateVersion;
+			}
+
+			if (!String.IsNullOrEmpty(selfUpdate.UpdateVersion))
+			{
+				OnUpdateReady?.Invoke();
+			}
+		}
+	}
+}
