UnrealToolbox: Add missing copyright notice.

#rnx

diff --git a/Engine/Source/Programs/UnrealToolbox/SelfUpdateService.cs b/Engine/Source/Programs/UnrealToolbox/SelfUpdateService.cs
--- a/Engine/Source/Programs/UnrealToolbox/SelfUpdateService.cs
+++ b/Engine/Source/Programs/UnrealToolbox/SelfUpdateService.cs
@@ -1,3 +1,5 @@
+// Copyright Epic Games, Inc. All Rights Reserved.
+
 using EpicGames.Core;
 using EpicGames.Horde.Storage.Nodes;
 using EpicGames.Horde.Tools;
