UnrealToolbox: Self-update fixes.

#rnx

diff --git a/Engine/Source/Programs/UnrealToolbox/SelfUpdateService.cs b/Engine/Source/Programs/UnrealToolbox/SelfUpdateService.cs
--- a/Engine/Source/Programs/UnrealToolbox/SelfUpdateService.cs
+++ b/Engine/Source/Programs/UnrealToolbox/SelfUpdateService.cs
@@ -119,7 +119,8 @@ namespace UnrealToolbox
 			}
 			else
 			{
-				if (!String.IsNullOrEmpty(_config.Current.LaunchApp))
+				string? updateVersion = _config.Current.UpdateVersion;
+				if (!String.IsNullOrEmpty(updateVersion) && !String.IsNullOrEmpty(_config.Current.LaunchApp))
 				{
 					return Launch(_config.Current.LaunchApp, _args);
 				}
@@ -132,8 +133,21 @@ namespace UnrealToolbox
 			using Process process = new Process();
 			if (executable.EndsWith(".dll", StringComparison.OrdinalIgnoreCase))
 			{
-				process.StartInfo.FileName = "dotnet";
-				process.StartInfo.ArgumentList.Add(executable);
+				string baseExecutable = executable.Substring(0, executable.Length - 4);
+				if (OperatingSystem.IsWindows())
+				{
+					baseExecutable += ".exe";
+				}
+
+				if (File.Exists(baseExecutable))
+				{
+					process.StartInfo.FileName = baseExecutable;
+				}
+				else
+				{
+					process.StartInfo.FileName = "dotnet";
+					process.StartInfo.ArgumentList.Add(executable);
+				}
 			}
 			else
 			{
@@ -220,7 +234,7 @@ namespace UnrealToolbox
 			_logger.LogInformation("Latest app version is {Id}", deployment.Id);
 
 			string updateVersion = deployment.Id.ToString();
-			if (String.Equals(updateVersion, selfUpdate.LatestVersion, StringComparison.OrdinalIgnoreCase) || String.Equals(updateVersion, selfUpdate.UpdateVersion, StringComparison.OrdinalIgnoreCase))
+			if (!String.Equals(updateVersion, selfUpdate.LatestVersion, StringComparison.OrdinalIgnoreCase) || !String.Equals(updateVersion, selfUpdate.UpdateVersion, StringComparison.OrdinalIgnoreCase))
 			{
 				selfUpdate.UpdateVersion = null;
 
