UnrealToolbox: Fix self update triggering infinite loop.

#rnx

diff --git a/Engine/Source/Programs/UnrealToolbox/SelfUpdateService.cs b/Engine/Source/Programs/UnrealToolbox/SelfUpdateService.cs
--- a/Engine/Source/Programs/UnrealToolbox/SelfUpdateService.cs
+++ b/Engine/Source/Programs/UnrealToolbox/SelfUpdateService.cs
@@ -234,7 +234,7 @@ namespace UnrealToolbox
 			_logger.LogInformation("Latest app version is {Id}", deployment.Id);
 
 			string updateVersion = deployment.Id.ToString();
-			if (!String.Equals(updateVersion, selfUpdate.LatestVersion, StringComparison.OrdinalIgnoreCase) || !String.Equals(updateVersion, selfUpdate.UpdateVersion, StringComparison.OrdinalIgnoreCase))
+			if (!String.Equals(updateVersion, selfUpdate.LatestVersion, StringComparison.OrdinalIgnoreCase) && !String.Equals(updateVersion, selfUpdate.UpdateVersion, StringComparison.OrdinalIgnoreCase))
 			{
 				selfUpdate.UpdateVersion = null;
 
