Unreal Toolbox: Disable the agent plugin if the agent is not installed. Status for this is refreshed whenever the app is focussed.

#rnx

diff --git a/Engine/Source/Programs/UnrealToolbox/App.axaml.cs b/Engine/Source/Programs/UnrealToolbox/App.axaml.cs
--- a/Engine/Source/Programs/UnrealToolbox/App.axaml.cs
+++ b/Engine/Source/Programs/UnrealToolbox/App.axaml.cs
@@ -169,6 +169,11 @@ namespace UnrealToolbox
 
 		private void TrayIcon_Click(object? sender, EventArgs e)
 		{
+			foreach(ITrayAppPlugin plugin in _serviceProvider.GetServices<ITrayAppPlugin>())
+			{
+				plugin.Refresh();
+			}
+
 			OpenSettings();
 		}
 
diff --git a/Engine/Source/Programs/UnrealToolbox/IToolboxPlugin.cs b/Engine/Source/Programs/UnrealToolbox/IToolboxPlugin.cs
--- a/Engine/Source/Programs/UnrealToolbox/IToolboxPlugin.cs
+++ b/Engine/Source/Programs/UnrealToolbox/IToolboxPlugin.cs
@@ -37,11 +37,21 @@ namespace UnrealToolbox
 		/// </summary>
 		IconSource Icon { get; }
 
+		/// <summary>
+		/// Refresh the state of this plugin. Called when the application is activated or focussed.
+		/// </summary>
+		void Refresh();
+
 		/// <summary>
 		/// Get the current status of this plugin
 		/// </summary>
 		TrayAppPluginStatus GetStatus();
 
+		/// <summary>
+		/// Determines if the plugin should be shown in the settings page
+		/// </summary>
+		bool HasSettingsPage();
+
 		/// <summary>
 		/// Create a settings page for this plugin
 		/// </summary>
diff --git a/Engine/Source/Programs/UnrealToolbox/Plugins/HordeAgent/HordeAgentPlugin.cs b/Engine/Source/Programs/UnrealToolbox/Plugins/HordeAgent/HordeAgentPlugin.cs
--- a/Engine/Source/Programs/UnrealToolbox/Plugins/HordeAgent/HordeAgentPlugin.cs
+++ b/Engine/Source/Programs/UnrealToolbox/Plugins/HordeAgent/HordeAgentPlugin.cs
@@ -1,15 +1,15 @@
 // Copyright Epic Games, Inc. All Rights Reserved.
 
-using System.Diagnostics;
-using System.IO.Pipes;
-using System.Runtime.InteropServices;
-using System.Text.Json;
-using System.Text.Json.Serialization;
 using Avalonia.Controls;
 using EpicGames.Core;
 using FluentAvalonia.UI.Controls;
 using HordeAgent;
 using Microsoft.Win32;
+using System.Diagnostics;
+using System.IO.Pipes;
+using System.Runtime.InteropServices;
+using System.Text.Json;
+using System.Text.Json.Serialization;
 
 namespace UnrealToolbox.Plugins.HordeAgent
 {
@@ -35,6 +35,11 @@ namespace UnrealToolbox.Plugins.HordeAgent
 
 		public IconSource Icon => new SymbolIconSource() { Symbol = Symbol.People };
 
+		public bool IsEnabled { get; private set; }
+
+		public bool HasSettingsPage()
+			=> IsEnabled;
+
 		public Control CreateSettingsPage(SettingsContext context)
 			=> new HordeAgentSettingsPage(context, this);
 
@@ -68,7 +73,7 @@ namespace UnrealToolbox.Plugins.HordeAgent
 
 			DirectoryReference? settingsRoot = DirectoryReference.GetSpecialFolder(Environment.SpecialFolder.CommonApplicationData);
 			settingsRoot ??= DirectoryReference.GetCurrentDirectory();
-			_settingsFile = FileReference.Combine(settingsRoot, "Epic", "Horde", "TrayApp", "Settings.json");
+			_settingsFile = FileReference.Combine(settingsRoot, "Epic", "Horde", "Agent", "Toolbox.json");
 
 			LoadSettings();
 
@@ -76,6 +81,14 @@ namespace UnrealToolbox.Plugins.HordeAgent
 			_tickPauseStateTask = BackgroundTask.StartNew(ctx => TickPauseStateAsync(ctx));
 		}
 
+		public void Refresh()
+		{
+			if (OperatingSystem.IsWindows())
+			{
+				IsEnabled = Registry.GetValue("HKEY_LOCAL_MACHINE\\SOFTWARE\\Epic Games\\Horde\\Agent", "Installed", null) != null;
+			}
+		}
+
 		static JsonSerializerOptions GetJsonSerializerOptions()
 		{
 			JsonSerializerOptions options = new JsonSerializerOptions();
@@ -142,9 +155,12 @@ namespace UnrealToolbox.Plugins.HordeAgent
 
 		public void PopulateContextMenu(NativeMenu contextMenu)
 		{
-			NativeMenuItem enrollMenuItem = new NativeMenuItem("Enroll with Server...");
-			enrollMenuItem.Click += (s, e) => EnrollWithServer();
-			contextMenu.Items.Add(enrollMenuItem);
+			if (IsEnabled)
+			{
+				NativeMenuItem enrollMenuItem = new NativeMenuItem("Enroll with Server...");
+				enrollMenuItem.Click += (s, e) => EnrollWithServer();
+				contextMenu.Items.Add(enrollMenuItem);
+			}
 		}
 
 		public async ValueTask DisposeAsync()
diff --git a/Engine/Source/Programs/UnrealToolbox/Plugins/HordeProxy/HordeProxyPlugin.cs b/Engine/Source/Programs/UnrealToolbox/Plugins/HordeProxy/HordeProxyPlugin.cs
--- a/Engine/Source/Programs/UnrealToolbox/Plugins/HordeProxy/HordeProxyPlugin.cs
+++ b/Engine/Source/Programs/UnrealToolbox/Plugins/HordeProxy/HordeProxyPlugin.cs
@@ -49,6 +49,9 @@ namespace UnrealToolbox.Plugins.HordeProxy
 			_serverTask = BackgroundTask.StartNew(RunServerAsync);
 		}
 
+		public bool HasSettingsPage()
+			=> true;
+
 		public Control CreateSettingsPage(SettingsContext context)
 			=> new HordeProxySettingsPage(this);
 
@@ -80,6 +83,10 @@ namespace UnrealToolbox.Plugins.HordeProxy
 			await _serverTask.DisposeAsync();
 		}
 
+		public void Refresh()
+		{
+		}
+
 		void OnStateChanged()
 		{
 			_restartServerEvent.Pulse();
diff --git a/Engine/Source/Programs/UnrealToolbox/SettingsWindow.axaml.cs b/Engine/Source/Programs/UnrealToolbox/SettingsWindow.axaml.cs
--- a/Engine/Source/Programs/UnrealToolbox/SettingsWindow.axaml.cs
+++ b/Engine/Source/Programs/UnrealToolbox/SettingsWindow.axaml.cs
@@ -1,6 +1,7 @@
 // Copyright Epic Games, Inc. All Rights Reserved.
 
 using Avalonia.Controls;
+using Avalonia.Input;
 using FluentAvalonia.UI.Controls;
 using Microsoft.Extensions.DependencyInjection;
 
@@ -28,9 +29,12 @@ namespace UnrealToolbox
 
 			foreach (ITrayAppPlugin plugin in _serviceProvider.GetServices<ITrayAppPlugin>())
 			{
-				string typeName = plugin.GetType().FullName!;
-				_navView.MenuItems.Add(new NavigationViewItem(){ Content = plugin.Name, IconSource = plugin.Icon, Tag = typeName });
-				_typeToPlugin.Add(typeName, plugin);
+				if(plugin.HasSettingsPage())
+				{
+					string typeName = plugin.GetType().FullName!;
+					_navView.MenuItems.Add(new NavigationViewItem() { Content = plugin.Name, IconSource = plugin.Icon, Tag = typeName });
+					_typeToPlugin.Add(typeName, plugin);
+				}
 			}
 
 			_navView.SelectedItem = _navView.MenuItems[0];
@@ -39,6 +43,19 @@ namespace UnrealToolbox
 			NavView_UpdateContent();
 		}
 
+		protected override void OnGotFocus(GotFocusEventArgs e)
+		{
+			base.OnGotFocus(e);
+
+			if (_serviceProvider != null)
+			{
+				foreach (ITrayAppPlugin plugin in _serviceProvider.GetServices<ITrayAppPlugin>())
+				{
+					plugin.Refresh();
+				}
+			}
+		}
+
 		private void NavView_SelectionChanged(object? sender, NavigationViewSelectionChangedEventArgs e)
 		{
 			if (sender == _navView)
