Horde: Unreal Toolbox notification service

diff --git a/Engine/Source/Programs/Horde/HordeAgent/IpcMessages.cs b/Engine/Source/Programs/Horde/HordeAgent/IpcMessages.cs
--- a/Engine/Source/Programs/Horde/HordeAgent/IpcMessages.cs
+++ b/Engine/Source/Programs/Horde/HordeAgent/IpcMessages.cs
@@ -174,16 +174,16 @@ namespace HordeAgent
 		/// <summary>
 		/// Static status object for starting an agent
 		/// </summary>
-		public static AgentStatusMessage Starting { get; } = new AgentStatusMessage(false, 0, "Starting up...");
+		public static AgentStatusMessage Starting { get; } = new AgentStatusMessage(true, 0, "Starting up...");
 
 		/// <summary>
 		/// Agent is waiting to be enrolled with the server
 		/// </summary>
-		public static AgentStatusMessage WaitingForEnrollment { get; } = new AgentStatusMessage(false, 0, "Waiting for enrollment...");
+		public static AgentStatusMessage WaitingForEnrollment { get; } = new AgentStatusMessage(true, 0, "Waiting for enrollment...");
 
 		/// <summary>
 		/// Agent is connecting to the server
 		/// </summary>
-		public static AgentStatusMessage ConnectingToServer { get; } = new AgentStatusMessage(false, 0, "Connecting to server...");
+		public static AgentStatusMessage ConnectingToServer { get; } = new AgentStatusMessage(true, 0, "Connecting to server...");
 	}
 }
diff --git a/Engine/Source/Programs/UnrealToolbox/App.axaml.cs b/Engine/Source/Programs/UnrealToolbox/App.axaml.cs
--- a/Engine/Source/Programs/UnrealToolbox/App.axaml.cs
+++ b/Engine/Source/Programs/UnrealToolbox/App.axaml.cs
@@ -4,12 +4,14 @@ using System.Diagnostics;
 using Avalonia;
 using Avalonia.Controls;
 using Avalonia.Controls.ApplicationLifetimes;
+using Avalonia.Interactivity;
 using Avalonia.Markup.Xaml;
 using Avalonia.Threading;
 using EpicGames.Core;
 using UnrealToolbox.Plugins.HordeAgent;
 using UnrealToolbox.Plugins.HordeProxy;
 using Microsoft.Extensions.DependencyInjection;
+using DesktopNotifications;
 
 namespace UnrealToolbox
 {
@@ -40,6 +42,9 @@ namespace UnrealToolbox
 			serviceCollection.AddSingleton<ITrayAppHost>(this);
 			serviceCollection.AddSingleton<SelfUpdateService>();
 
+			serviceCollection.AddSingleton<ToolboxNotificationManager>();
+			serviceCollection.AddSingleton<IToolboxNotificationManager, ToolboxNotificationManager>(sp => sp.GetRequiredService<ToolboxNotificationManager>());
+
 			// Configure plugins
 			serviceCollection.AddSingleton<ITrayAppPlugin, HordeAgentPlugin>();
 			serviceCollection.AddSingleton<ITrayAppPlugin, HordeProxyPlugin>();
@@ -66,6 +71,9 @@ namespace UnrealToolbox
 			SelfUpdateService selfUpdateService = _serviceProvider.GetRequiredService<SelfUpdateService>();
 			selfUpdateService.OnUpdateReady += UpdateReady;
 			selfUpdateService.Start();
+
+			ToolboxNotificationManager notificationManager = _serviceProvider.GetRequiredService<ToolboxNotificationManager>();
+			notificationManager.Start();
 		}
 
 		private void UpdateReady()
@@ -217,10 +225,18 @@ namespace UnrealToolbox
 		{
 			TrayAppPluginState state = TrayAppPluginState.Undefined;
 			List<string> messages = new List<string>();
+			string? pluginName = null;
+			string? errorMessage = null; 
 
 			foreach (ITrayAppPlugin plugin in _serviceProvider.GetServices<ITrayAppPlugin>())
 			{
 				TrayAppPluginStatus status = plugin.GetStatus();
+				if (status.State == TrayAppPluginState.Error && !String.IsNullOrEmpty(status.NotificationMessage))
+				{
+					pluginName = plugin.Name;
+					errorMessage = status.NotificationMessage;
+				}
+
 				if (status.State >= state)
 				{
 					if (status.State > state)
@@ -244,6 +260,16 @@ namespace UnrealToolbox
 				_ => _normalIcon
 			};
 			trayIcon.ToolTipText = String.Join("\n", messages);
+
+			// Notify on errors
+			if (!String.IsNullOrEmpty(errorMessage))
+			{
+				IToolboxNotificationManager? notifications = _serviceProvider.GetService<IToolboxNotificationManager>();
+				if (notifications != null)
+				{
+					notifications.ShowNotification(pluginName ?? "Status Error", errorMessage);
+				}
+			}
 		}
 	}
 }
diff --git a/Engine/Source/Programs/UnrealToolbox/IToolboxPlugin.cs b/Engine/Source/Programs/UnrealToolbox/IToolboxPlugin.cs
--- a/Engine/Source/Programs/UnrealToolbox/IToolboxPlugin.cs
+++ b/Engine/Source/Programs/UnrealToolbox/IToolboxPlugin.cs
@@ -17,9 +17,9 @@ namespace UnrealToolbox
 	/// <summary>
 	/// Reported status of a plugin
 	/// </summary>
-	record class TrayAppPluginStatus(TrayAppPluginState State, string? Message)
+	record class TrayAppPluginStatus(TrayAppPluginState State, string? Message = null, string? NotificationMessage = null)
 	{
-		public static TrayAppPluginStatus Default { get; } = new TrayAppPluginStatus(TrayAppPluginState.Undefined, null);
+		public static TrayAppPluginStatus Default { get; } = new TrayAppPluginStatus(TrayAppPluginState.Undefined);
 	}
 
 	/// <summary>
diff --git a/Engine/Source/Programs/UnrealToolbox/Notifications/Notifications.cs b/Engine/Source/Programs/UnrealToolbox/Notifications/Notifications.cs
new file mode 100644
--- /dev/null
+++ b/Engine/Source/Programs/UnrealToolbox/Notifications/Notifications.cs
@@ -0,0 +1,103 @@
+// Copyright Epic Games, Inc. All Rights Reserved.
+
+using Avalonia;
+using Avalonia.Controls;
+using Avalonia.Controls.ApplicationLifetimes;
+using Avalonia.Platform;
+using DesktopNotifications;
+using DesktopNotifications.FreeDesktop;
+using DesktopNotifications.Windows;
+using Microsoft.Extensions.Logging;
+using System;
+using UnrealToolbox;
+
+namespace UnrealToolbox
+{
+	/// <summary>
+	/// Toolbox notification manager interface
+	/// </summary>
+	public interface IToolboxNotificationManager 
+	{
+		void ShowNotification(string title, string body, bool force = false);
+	}
+
+	/// <summary>
+	/// Toolbox notification manager implementation
+	/// </summary>
+	class ToolboxNotificationManager : IToolboxNotificationManager, IDisposable
+	{
+
+		private INotificationManager? _platformManager;
+		readonly ILogger _logger;
+
+		// spam prevention
+		DateTime? _lastNotificationTime;
+		string? _lastTitle;
+		string? _lastBody;
+
+		public ToolboxNotificationManager(ILogger<ToolboxNotificationManager> logger)
+		{
+			_logger = logger;
+		}
+
+		public void Start()
+		{
+			if (Environment.OSVersion.Platform == PlatformID.Win32NT)
+			{
+				WindowsApplicationContext context = WindowsApplicationContext.FromCurrentProcess();
+				_platformManager = new WindowsNotificationManager(context);
+			}
+			else
+			{
+				throw new NotImplementedException();
+			}
+
+			// initialize and ensure with result
+			_platformManager.Initialize().GetAwaiter().GetResult();
+
+		}
+
+		public void Dispose()
+		{
+			_platformManager?.Dispose();
+		}
+
+		/// <summary>
+		/// Show a notification, 
+		/// </summary>
+		/// <param name="title"></param>
+		/// <param name="body"></param>
+		public void ShowNotification(string title, string body, bool force = false)
+		{
+			// spawn 
+			if (!force && _lastNotificationTime != null && !String.IsNullOrEmpty(_lastBody) && !String.IsNullOrEmpty(_lastTitle)) 
+			{
+				TimeSpan deltaTime = DateTime.Now - _lastNotificationTime.Value; 
+
+				// don't show a new notification if already displayed one in last minute
+				if (deltaTime.TotalSeconds < 60)
+				{
+					return;
+				}
+
+				// if the title and body are the same, wait 2 minutes
+				if ((deltaTime.TotalSeconds < 120) && title == _lastTitle && body == _lastBody)
+				{
+					return;
+				}
+			}
+
+			_lastTitle = title;
+			_lastBody = body;
+			_lastNotificationTime = DateTime.Now;
+
+			Notification notification = new Notification
+			{
+				Title = title,
+				Body = body
+			};
+
+			_platformManager?.ShowNotification(notification, DateTimeOffset.Now + TimeSpan.FromSeconds(30));
+		}
+	}
+}
\ No newline at end of file
diff --git a/Engine/Source/Programs/UnrealToolbox/Plugins/HordeAgent/HordeAgentPlugin.cs b/Engine/Source/Programs/UnrealToolbox/Plugins/HordeAgent/HordeAgentPlugin.cs
--- a/Engine/Source/Programs/UnrealToolbox/Plugins/HordeAgent/HordeAgentPlugin.cs
+++ b/Engine/Source/Programs/UnrealToolbox/Plugins/HordeAgent/HordeAgentPlugin.cs
@@ -189,7 +189,7 @@ namespace UnrealToolbox.Plugins.HordeAgent
 			if (!status.Healthy)
 			{
 				string message = String.IsNullOrEmpty(status.Detail) ? "Error. Check logs." : status.Detail.Length > 100 ? status.Detail.Substring(0, 100) : status.Detail;
-				_status = new TrayAppPluginStatus(TrayAppPluginState.Error, message);
+				_status = new TrayAppPluginStatus(TrayAppPluginState.Error, message, message);
 			}
 			else if (status.NumLeases > 0)
 			{
@@ -442,10 +442,10 @@ namespace UnrealToolbox.Plugins.HordeAgent
 			AgentMessageBuffer message = new AgentMessageBuffer();
 			using (NamedPipeClientStream pipeClient = new NamedPipeClientStream(".", AgentMessagePipe.PipeName, PipeDirection.InOut))
 			{
-				SetStatus(new AgentStatusMessage(false, 0, "Connecting to agent..."));
+				SetStatus(new AgentStatusMessage(true, 0, "Connecting to agent..."));
 				await pipeClient.ConnectAsync(cancellationToken);
 
-				SetStatus(new AgentStatusMessage(false, 0, "Waiting for status update."));
+				SetStatus(new AgentStatusMessage(true, 0, "Waiting for status update."));
 				for (; ; )
 				{
 					Task idleChangeTask = _enabledChangedEvent.Task;
diff --git a/Engine/Source/Programs/UnrealToolbox/UnrealToolbox.csproj b/Engine/Source/Programs/UnrealToolbox/UnrealToolbox.csproj
--- a/Engine/Source/Programs/UnrealToolbox/UnrealToolbox.csproj
+++ b/Engine/Source/Programs/UnrealToolbox/UnrealToolbox.csproj
@@ -6,7 +6,8 @@
 
   <PropertyGroup>
     <OutputType>WinExe</OutputType>
-    <TargetFramework>net8.0-windows</TargetFramework>
+    <TargetFramework>net8.0-windows10.0.17763.0</TargetFramework>
+    <WindowsSdkPackageVersion>10.0.17763.41</WindowsSdkPackageVersion>
     <Nullable>enable</Nullable>
     <ImplicitUsings>enable</ImplicitUsings>
     <NeutralLanguage>en-US</NeutralLanguage>
@@ -14,6 +15,7 @@
 	<AssemblyName>UnrealToolbox</AssemblyName>
     <AvaloniaUseCompiledBindingsByDefault>true</AvaloniaUseCompiledBindingsByDefault>
     <NoWarn>IDE0005;CA1501</NoWarn>
+    <ApplicationIcon>Resources/StatusNormal.ico</ApplicationIcon>
   </PropertyGroup>
 
   <PropertyGroup Condition=" '$(Configuration)|$(Platform)' == 'Analyze|AnyCPU' ">
@@ -72,6 +74,7 @@
     <PackageReference Condition="'$(Configuration)' == 'Debug'" Include="Avalonia.Diagnostics" Version="11.1.0" />
     <PackageReference Include="FluentAvaloniaUI" Version="2.1.0" />
     <PackageReference Include="CommunityToolkit.Mvvm" Version="8.3.1" />
+    <PackageReference Include="DesktopNotifications.Avalonia" Version="1.3.1" />
   </ItemGroup>
 
 </Project>
\ No newline at end of file
