Horde: Fix warnings
#rnx

diff --git a/Engine/Source/Programs/UnrealToolbox/Notifications/Notifications.cs b/Engine/Source/Programs/UnrealToolbox/Notifications/Notifications.cs
--- a/Engine/Source/Programs/UnrealToolbox/Notifications/Notifications.cs
+++ b/Engine/Source/Programs/UnrealToolbox/Notifications/Notifications.cs
@@ -18,6 +18,12 @@ namespace UnrealToolbox
 	/// </summary>
 	public interface IToolboxNotificationManager 
 	{
+		/// <summary>
+		/// Show a notification
+		/// </summary>
+		/// <param name="title"></param>
+		/// <param name="body"></param>
+		/// <param name="force"></param>
 		void ShowNotification(string title, string body, bool force = false);
 	}
 
@@ -67,6 +73,7 @@ namespace UnrealToolbox
 		/// </summary>
 		/// <param name="title"></param>
 		/// <param name="body"></param>
+		/// <param name="force"></param>
 		public void ShowNotification(string title, string body, bool force = false)
 		{
 			// spawn 
