Horde: Store proxy server settings in AppData\Local\Epic Games\Horde\Proxy\Settings.json, rather than C:\ProgramData.

diff --git a/Engine/Source/Programs/UnrealToolbox/Plugins/HordeProxy/HordeProxyPlugin.cs b/Engine/Source/Programs/UnrealToolbox/Plugins/HordeProxy/HordeProxyPlugin.cs
--- a/Engine/Source/Programs/UnrealToolbox/Plugins/HordeProxy/HordeProxyPlugin.cs
+++ b/Engine/Source/Programs/UnrealToolbox/Plugins/HordeProxy/HordeProxyPlugin.cs
@@ -41,9 +41,9 @@ namespace UnrealToolbox.Plugins.HordeProxy
 			_hordeClientProvider.OnStateChanged += OnStateChanged;
 			_synchronizationContext = SynchronizationContext.Current;
 
-			DirectoryReference ? settingsRoot = DirectoryReference.GetSpecialFolder(Environment.SpecialFolder.CommonApplicationData);
+			DirectoryReference? settingsRoot = DirectoryReference.GetSpecialFolder(Environment.SpecialFolder.LocalApplicationData);
 			settingsRoot ??= DirectoryReference.GetCurrentDirectory();
-			_settings = new JsonConfig<HordeProxySettings>(FileReference.Combine(settingsRoot, "Epic", "Horde", "Proxy", "Settings.json"));
+			_settings = new JsonConfig<HordeProxySettings>(FileReference.Combine(settingsRoot, "Epic Games", "Horde", "Proxy", "Settings.json"));
 			_settings.LoadSettings();
 
 			_serverTask = BackgroundTask.StartNew(RunServerAsync);
