Unreal Toolbox: Fix icon showing status of "Agent is paused" when agent plugin is disabled.

diff --git a/Engine/Source/Programs/UnrealToolbox/Plugins/HordeAgent/HordeAgentPlugin.cs b/Engine/Source/Programs/UnrealToolbox/Plugins/HordeAgent/HordeAgentPlugin.cs
--- a/Engine/Source/Programs/UnrealToolbox/Plugins/HordeAgent/HordeAgentPlugin.cs
+++ b/Engine/Source/Programs/UnrealToolbox/Plugins/HordeAgent/HordeAgentPlugin.cs
@@ -71,9 +71,9 @@ namespace UnrealToolbox.Plugins.HordeAgent
 		{
 			_host = host;
 
-			DirectoryReference? settingsRoot = DirectoryReference.GetSpecialFolder(Environment.SpecialFolder.CommonApplicationData);
+			DirectoryReference? settingsRoot = DirectoryReference.GetSpecialFolder(Environment.SpecialFolder.LocalApplicationData);
 			settingsRoot ??= DirectoryReference.GetCurrentDirectory();
-			_settingsFile = FileReference.Combine(settingsRoot, "Epic", "Horde", "Agent", "Toolbox.json");
+			_settingsFile = FileReference.Combine(settingsRoot, "Epic Games", "Horde", "Agent", "Toolbox.json");
 
 			LoadSettings();
 
@@ -186,7 +186,11 @@ namespace UnrealToolbox.Plugins.HordeAgent
 
 		void SetStatus(AgentStatusMessage status)
 		{
-			if (!status.Healthy)
+			if (!IsEnabled)
+			{
+				_status = null;
+			}
+			else if (!status.Healthy)
 			{
 				string message = String.IsNullOrEmpty(status.Detail) ? "Error. Check logs." : status.Detail.Length > 100 ? status.Detail.Substring(0, 100) : status.Detail;
 				_status = new TrayAppPluginStatus(TrayAppPluginState.Error, message, message);
