Unreal Toolbox: Store tools in AppData\Local\Epic Games\Unreal Toolbox\Tools.

diff --git a/Engine/Source/Programs/UnrealToolbox/ToolCatalog.cs b/Engine/Source/Programs/UnrealToolbox/ToolCatalog.cs
--- a/Engine/Source/Programs/UnrealToolbox/ToolCatalog.cs
+++ b/Engine/Source/Programs/UnrealToolbox/ToolCatalog.cs
@@ -236,7 +236,7 @@ namespace UnrealToolbox
 
 			_hordeClientProvider = hordeClientProvider;
 			_logger = logger;
-			_baseDir = DirectoryReference.Combine(toolsDir, "Epic Games", "Horde", "TrayApp");
+			_baseDir = DirectoryReference.Combine(toolsDir, "Epic Games", "Unreal Toolbox");
 			_configFile = FileReference.Combine(_baseDir, "Tools.json");
 			_updateEvent = new AsyncEvent();
 			_updateTask = new BackgroundTask(UpdateAsync);
@@ -342,7 +342,7 @@ namespace UnrealToolbox
 
 		DirectoryReference GetToolDir(ToolId toolId, ToolDeploymentId deploymentId)
 		{
-			return DirectoryReference.Combine(_baseDir, toolId.ToString(), deploymentId.ToString());
+			return DirectoryReference.Combine(_baseDir, "Tools", toolId.ToString(), deploymentId.ToString());
 		}
 
 		ToolConfig LoadToolConfig(DirectoryReference toolDir)
