Unreal Toolbox: Move settings for Horde Proxy under the Unreal Toolbox folder.

diff --git a/Engine/Source/Programs/UnrealToolbox/Plugins/HordeProxy/HordeProxyPlugin.cs b/Engine/Source/Programs/UnrealToolbox/Plugins/HordeProxy/HordeProxyPlugin.cs
--- a/Engine/Source/Programs/UnrealToolbox/Plugins/HordeProxy/HordeProxyPlugin.cs
+++ b/Engine/Source/Programs/UnrealToolbox/Plugins/HordeProxy/HordeProxyPlugin.cs
@@ -43,7 +43,7 @@ namespace UnrealToolbox.Plugins.HordeProxy
 
 			DirectoryReference? settingsRoot = DirectoryReference.GetSpecialFolder(Environment.SpecialFolder.LocalApplicationData);
 			settingsRoot ??= DirectoryReference.GetCurrentDirectory();
-			_settings = new JsonConfig<HordeProxySettings>(FileReference.Combine(settingsRoot, "Epic Games", "Horde", "Proxy", "Settings.json"));
+			_settings = new JsonConfig<HordeProxySettings>(FileReference.Combine(settingsRoot, "Epic Games", "Unreal Toolbox", "HordeProxy.json"));
 			_settings.LoadSettings();
 
 			_serverTask = BackgroundTask.StartNew(RunServerAsync);
