Unreal Toolbox: Move Horde Agent settings under AppData\Local\Unreal Toolbox\HordeAgent.json.

diff --git a/Engine/Source/Programs/UnrealToolbox/Plugins/HordeAgent/HordeAgentPlugin.cs b/Engine/Source/Programs/UnrealToolbox/Plugins/HordeAgent/HordeAgentPlugin.cs
--- a/Engine/Source/Programs/UnrealToolbox/Plugins/HordeAgent/HordeAgentPlugin.cs
+++ b/Engine/Source/Programs/UnrealToolbox/Plugins/HordeAgent/HordeAgentPlugin.cs
@@ -4,6 +4,7 @@ using Avalonia.Controls;
 using EpicGames.Core;
 using FluentAvalonia.UI.Controls;
 using HordeAgent;
+using Microsoft.Extensions.Logging;
 using Microsoft.Win32;
 using System.Diagnostics;
 using System.IO.Pipes;
@@ -18,6 +19,7 @@ namespace UnrealToolbox.Plugins.HordeAgent
 		record struct IdleStat(string Name, long Value, long MinValue);
 
 		readonly ITrayAppHost _host;
+		readonly ILogger _logger;
 
 		readonly BackgroundTask _clientTask;
 		readonly BackgroundTask _tickPauseStateTask;
@@ -67,13 +69,14 @@ namespace UnrealToolbox.Plugins.HordeAgent
 			}
 		}
 
-		public HordeAgentPlugin(ITrayAppHost host)
+		public HordeAgentPlugin(ITrayAppHost host, ILogger<HordeAgentPlugin> logger)
 		{
 			_host = host;
+			_logger = logger;
 
 			DirectoryReference? settingsRoot = DirectoryReference.GetSpecialFolder(Environment.SpecialFolder.LocalApplicationData);
 			settingsRoot ??= DirectoryReference.GetCurrentDirectory();
-			_settingsFile = FileReference.Combine(settingsRoot, "Epic Games", "Horde", "Agent", "Toolbox.json");
+			_settingsFile = FileReference.Combine(settingsRoot, "Epic Games", "Unreal Toolbox", "HordeAgent.json");
 
 			LoadSettings();
 
@@ -120,8 +123,9 @@ namespace UnrealToolbox.Plugins.HordeAgent
 						return;
 					}
 				}
-				catch (Exception)
+				catch (Exception ex)
 				{
+					_logger.LogError(ex, "Error while reading settings: {Message}", ex.Message);
 				}
 			}
 
