Unreal Toolbox: Read legacy settings file for Horde Agent if an updated settings file is not found.

diff --git a/Engine/Source/Programs/UnrealToolbox/Plugins/HordeAgent/HordeAgentPlugin.cs b/Engine/Source/Programs/UnrealToolbox/Plugins/HordeAgent/HordeAgentPlugin.cs
--- a/Engine/Source/Programs/UnrealToolbox/Plugins/HordeAgent/HordeAgentPlugin.cs
+++ b/Engine/Source/Programs/UnrealToolbox/Plugins/HordeAgent/HordeAgentPlugin.cs
@@ -125,7 +125,32 @@ namespace UnrealToolbox.Plugins.HordeAgent
 				}
 				catch (Exception ex)
 				{
-					_logger.LogError(ex, "Error while reading settings: {Message}", ex.Message);
+					_logger.LogError(ex, "Error while reading {File}: {Message}", _settingsFile, ex.Message);
+				}
+			}
+			else
+			{
+				DirectoryReference? programData = DirectoryReference.GetSpecialFolder(Environment.SpecialFolder.CommonApplicationData);
+				if (programData != null)
+				{
+					FileReference legacySettingsFile = FileReference.Combine(programData, "Epic", "Horde", "TrayApp", "settings.json");
+					if (FileReference.Exists(legacySettingsFile))
+					{
+						try
+						{
+							byte[] data = FileReference.ReadAllBytes(legacySettingsFile);
+
+							HordeAgentSettings settings = JsonSerializer.Deserialize<HordeAgentSettings>(data, GetJsonSerializerOptions())!;
+							settings.Mode ??= ReadLegacyMode();
+
+							UpdateSettings(settings);
+							return;
+						}
+						catch (Exception ex)
+						{
+							_logger.LogError(ex, "Error while reading {File}: {Message}", legacySettingsFile, ex.Message);
+						}
+					}
 				}
 			}
 
