Unreal Toolbox: Prevent toolbox trying to delete it's own update files.

diff --git a/Engine/Source/Programs/UnrealToolbox/ToolCatalog.cs b/Engine/Source/Programs/UnrealToolbox/ToolCatalog.cs
--- a/Engine/Source/Programs/UnrealToolbox/ToolCatalog.cs
+++ b/Engine/Source/Programs/UnrealToolbox/ToolCatalog.cs
@@ -193,6 +193,7 @@ namespace UnrealToolbox
 		readonly ILogger _logger;
 		readonly object _workerThreadLockObject = new object();
 		readonly DirectoryReference _baseDir;
+		readonly DirectoryReference _toolsDir;
 		readonly FileReference _configFile;
 		readonly AsyncEvent _updateEvent;
 		readonly BackgroundTask _updateTask;
@@ -210,8 +211,6 @@ namespace UnrealToolbox
 		Uri? _serverUri;
 		bool _autoUpdate = true;
 
-		public DirectoryReference BaseDir => _baseDir;
-
 		public event Action? OnItemsChanged;
 
 		static readonly JsonSerializerOptions s_jsonSerializerOptions = new JsonSerializerOptions { AllowTrailingCommas = true, PropertyNameCaseInsensitive = true, PropertyNamingPolicy = JsonNamingPolicy.CamelCase, WriteIndented = true };
@@ -237,6 +236,7 @@ namespace UnrealToolbox
 			_hordeClientProvider = hordeClientProvider;
 			_logger = logger;
 			_baseDir = DirectoryReference.Combine(toolsDir, "Epic Games", "Unreal Toolbox");
+			_toolsDir = DirectoryReference.Combine(_baseDir, "Tools");
 			_configFile = FileReference.Combine(_baseDir, "Tools.json");
 			_updateEvent = new AsyncEvent();
 			_updateTask = new BackgroundTask(UpdateAsync);
@@ -342,7 +342,7 @@ namespace UnrealToolbox
 
 		DirectoryReference GetToolDir(ToolId toolId, ToolDeploymentId deploymentId)
 		{
-			return DirectoryReference.Combine(_baseDir, "Tools", toolId.ToString(), deploymentId.ToString());
+			return DirectoryReference.Combine(_toolsDir, toolId.ToString(), deploymentId.ToString());
 		}
 
 		ToolConfig LoadToolConfig(DirectoryReference toolDir)
@@ -617,7 +617,7 @@ namespace UnrealToolbox
 			List<DirectoryInfo> deleteToolDirs = new List<DirectoryInfo>();
 			List<DirectoryInfo> deleteDeploymentDirs = new List<DirectoryInfo>();
 
-			DirectoryInfo baseDirInfo = _baseDir.ToDirectoryInfo();
+			DirectoryInfo baseDirInfo = _toolsDir.ToDirectoryInfo();
 			if (baseDirInfo.Exists)
 			{
 				foreach (DirectoryInfo toolDir in baseDirInfo.EnumerateDirectories())
