Unreal Toolbox: Add -Settings and -Close command line arguments.

* -Settings will open the settings panel, either on the currently running instance or new instance.
* -Close will cause any currently running instance to shutdown.

The -Settings argument is specified when launching from the start menu. The -Close argument is passed to shut down the current instance when uninstalling.

diff --git a/Engine/Source/Programs/UnrealToolbox/App.axaml.cs b/Engine/Source/Programs/UnrealToolbox/App.axaml.cs
--- a/Engine/Source/Programs/UnrealToolbox/App.axaml.cs
+++ b/Engine/Source/Programs/UnrealToolbox/App.axaml.cs
@@ -18,7 +18,7 @@ namespace UnrealToolbox
 	/// <summary>
 	/// Main application class
 	/// </summary>
-	public sealed partial class App : Application, ITrayAppHost
+	public sealed partial class App : Application, ITrayAppHost, IAsyncDisposable
 	{
 		readonly ServiceProvider _serviceProvider;
 
@@ -29,6 +29,9 @@ namespace UnrealToolbox
 
 		SettingsWindow? _settingsWindow;
 
+		Thread? _settingsThread;
+		ManualResetEvent? _settingsThreadStop;
+
 		/// <summary>
 		/// Constructor
 		/// </summary>
@@ -52,6 +55,18 @@ namespace UnrealToolbox
 			_serviceProvider = serviceCollection.BuildServiceProvider();
 		}
 
+		public async ValueTask DisposeAsync()
+		{
+			if (_settingsThreadStop != null)
+			{
+				_settingsThreadStop.Set();
+				_settingsThread?.Join();
+				_settingsThreadStop.Dispose();
+				_settingsThreadStop = null;
+			}
+			await _serviceProvider.DisposeAsync();
+		}
+
 		/// <inheritdoc/>
 		public override void Initialize()
 		{
@@ -74,6 +89,37 @@ namespace UnrealToolbox
 
 			ToolboxNotificationManager notificationManager = _serviceProvider.GetRequiredService<ToolboxNotificationManager>();
 			notificationManager.Start();
+
+			_settingsThreadStop = new ManualResetEvent(false);
+			_settingsThread = new Thread(WaitForEvents);
+			_settingsThread.Start();
+		}
+
+		private void WaitForEvents()
+		{
+			using EventWaitHandle closeEvent = new EventWaitHandle(false, EventResetMode.AutoReset, Program.CloseEventName);
+			using EventWaitHandle settingsEvent = new EventWaitHandle(false, EventResetMode.AutoReset, Program.SettingsEventName);
+			for(; ;)
+			{
+				int index = WaitHandle.WaitAny(new[] { closeEvent, settingsEvent, _settingsThreadStop! });
+				if (index == 0)
+				{
+					Dispatcher.UIThread.Post(() => CloseMainThread());
+				}
+				else if (index == 1)
+				{
+					Dispatcher.UIThread.Post(() => OpenSettings());
+				}
+				else
+				{
+					break;
+				}
+			}
+		}
+
+		void CloseMainThread()
+		{
+			((IClassicDesktopStyleApplicationLifetime)ApplicationLifetime!).Shutdown();
 		}
 
 		private void UpdateReady()
diff --git a/Engine/Source/Programs/UnrealToolbox/Program.cs b/Engine/Source/Programs/UnrealToolbox/Program.cs
--- a/Engine/Source/Programs/UnrealToolbox/Program.cs
+++ b/Engine/Source/Programs/UnrealToolbox/Program.cs
@@ -12,7 +12,8 @@ namespace UnrealToolbox
 	static class Program
 	{
 		const string MutexName = "UnrealToolbox-Mutex";
-		const string EventName = "UnrealToolbox-Exit";
+		public const string CloseEventName = "UnrealToolbox-Close";
+		public const string SettingsEventName = "UnrealToolbox-Settings";
 
 		public static SelfUpdateState? Update { get; private set; }
 
@@ -41,14 +42,36 @@ namespace UnrealToolbox
 
 		static int RealMain(string[] args)
 		{
+			using EventWaitHandle closeEvent = new EventWaitHandle(false, EventResetMode.AutoReset, CloseEventName);
+			if (args.Any(x => x.Equals("-Close", StringComparison.OrdinalIgnoreCase)))
+			{
+				closeEvent.Set();
+			}
+
+			using EventWaitHandle settingsEvent = new EventWaitHandle(false, EventResetMode.AutoReset, SettingsEventName);
+			if (args.Any(x => x.Equals("-Settings", StringComparison.OrdinalIgnoreCase)))
+			{
+				settingsEvent.Set();
+			}
+
 			using SingleInstanceMutex mutex = new SingleInstanceMutex(MutexName);
 			if (!mutex.Wait(0))
 			{
 				return 1;
 			}
 
-			BuildAvaloniaApp()
-				.StartWithClassicDesktopLifetime(args, ShutdownMode.OnExplicitShutdown);
+			AppBuilder builder = BuildAvaloniaApp();
+			try
+			{
+				builder.StartWithClassicDesktopLifetime(args, ShutdownMode.OnExplicitShutdown);
+			}
+			finally
+			{
+				if (builder.Instance is App app)
+				{
+					Task.Run(async () => await app.DisposeAsync()).Wait();
+				}
+			}
 
 			return 0;
 		}
