UnrealToolbox: Various improvements with installation process.

* Running UnrealToolbox with -Refresh argument will trigger a refresh of all plugins, allowing them to activate/deactivate.
* Refresh is run at the end of agent installation, allowing the UI to update to include menu item for enrolling agent, and adding UI panel for settings.
* Toolbox is now installed to C:\Program Files\Epic Games\Unreal Toolbox.
* Toolbox start menu item is no longer in a Horde folder.

diff --git a/Engine/Source/Programs/UnrealToolbox/App.axaml.cs b/Engine/Source/Programs/UnrealToolbox/App.axaml.cs
--- a/Engine/Source/Programs/UnrealToolbox/App.axaml.cs
+++ b/Engine/Source/Programs/UnrealToolbox/App.axaml.cs
@@ -77,6 +77,7 @@ namespace UnrealToolbox
 			_pausedIcon = (WindowIcon)Resources["StatusPaused"]!;
 			_errorIcon = (WindowIcon)Resources["StatusError"]!;
 
+			RefreshPlugins();
 			UpdateMenu();
 
 			ToolCatalog toolCatalog = _serviceProvider.GetRequiredService<ToolCatalog>();
@@ -98,15 +99,20 @@ namespace UnrealToolbox
 		private void WaitForEvents()
 		{
 			using EventWaitHandle closeEvent = new EventWaitHandle(false, EventResetMode.AutoReset, Program.CloseEventName);
+			using EventWaitHandle refreshEvent = new EventWaitHandle(false, EventResetMode.AutoReset, Program.RefreshEventName);
 			using EventWaitHandle settingsEvent = new EventWaitHandle(false, EventResetMode.AutoReset, Program.SettingsEventName);
-			for(; ;)
+			for (; ;)
 			{
-				int index = WaitHandle.WaitAny(new[] { closeEvent, settingsEvent, _settingsThreadStop! });
+				int index = WaitHandle.WaitAny(new[] { closeEvent, refreshEvent, settingsEvent, _settingsThreadStop! });
 				if (index == 0)
 				{
 					Dispatcher.UIThread.Post(() => CloseMainThread());
 				}
 				else if (index == 1)
+				{
+					Dispatcher.UIThread.Post(() => RefreshPlugins());
+				}
+				else if (index == 2)
 				{
 					Dispatcher.UIThread.Post(() => OpenSettings());
 				}
@@ -221,13 +227,23 @@ namespace UnrealToolbox
 			Process.Start(startInfo);
 		}
 
-		private void TrayIcon_Click(object? sender, EventArgs e)
+		private void RefreshPlugins()
 		{
-			foreach(ITrayAppPlugin plugin in _serviceProvider.GetServices<ITrayAppPlugin>())
+			bool update = false;
+			foreach (ITrayAppPlugin plugin in _serviceProvider.GetServices<ITrayAppPlugin>())
 			{
-				plugin.Refresh();
+				update |= plugin.Refresh();
 			}
+			if (update)
+			{
+				UpdateMenu();
+				_settingsWindow?.Refresh();
+			}
+		}
 
+		private void TrayIcon_Click(object? sender, EventArgs e)
+		{
+			RefreshPlugins();
 			OpenSettings();
 		}
 
diff --git a/Engine/Source/Programs/UnrealToolbox/IToolboxPlugin.cs b/Engine/Source/Programs/UnrealToolbox/IToolboxPlugin.cs
--- a/Engine/Source/Programs/UnrealToolbox/IToolboxPlugin.cs
+++ b/Engine/Source/Programs/UnrealToolbox/IToolboxPlugin.cs
@@ -40,7 +40,8 @@ namespace UnrealToolbox
 		/// <summary>
 		/// Refresh the state of this plugin. Called when the application is activated or focussed.
 		/// </summary>
-		void Refresh();
+		/// <returns>True if the plugin status has changed, and the context menu needs to be rebuilt</returns>
+		bool Refresh();
 
 		/// <summary>
 		/// Get the current status of this plugin
diff --git a/Engine/Source/Programs/UnrealToolbox/Plugins/HordeAgent/HordeAgentPlugin.cs b/Engine/Source/Programs/UnrealToolbox/Plugins/HordeAgent/HordeAgentPlugin.cs
--- a/Engine/Source/Programs/UnrealToolbox/Plugins/HordeAgent/HordeAgentPlugin.cs
+++ b/Engine/Source/Programs/UnrealToolbox/Plugins/HordeAgent/HordeAgentPlugin.cs
@@ -84,12 +84,18 @@ namespace UnrealToolbox.Plugins.HordeAgent
 			_tickPauseStateTask = BackgroundTask.StartNew(ctx => TickPauseStateAsync(ctx));
 		}
 
-		public void Refresh()
+		public bool Refresh()
 		{
 			if (OperatingSystem.IsWindows())
 			{
-				IsEnabled = Registry.GetValue("HKEY_LOCAL_MACHINE\\SOFTWARE\\Epic Games\\Horde\\Agent", "Installed", null) != null;
+				bool enabled = ((Registry.GetValue("HKEY_LOCAL_MACHINE\\SOFTWARE\\Epic Games\\Horde\\Agent", "Installed", null) as int?) ?? 0) != 0;
+				if (enabled != IsEnabled)
+				{
+					IsEnabled = enabled;
+					return true;
+				}
 			}
+			return false;
 		}
 
 		static JsonSerializerOptions GetJsonSerializerOptions()
diff --git a/Engine/Source/Programs/UnrealToolbox/Plugins/HordeProxy/HordeProxyPlugin.cs b/Engine/Source/Programs/UnrealToolbox/Plugins/HordeProxy/HordeProxyPlugin.cs
--- a/Engine/Source/Programs/UnrealToolbox/Plugins/HordeProxy/HordeProxyPlugin.cs
+++ b/Engine/Source/Programs/UnrealToolbox/Plugins/HordeProxy/HordeProxyPlugin.cs
@@ -83,9 +83,8 @@ namespace UnrealToolbox.Plugins.HordeProxy
 			await _serverTask.DisposeAsync();
 		}
 
-		public void Refresh()
-		{
-		}
+		public bool Refresh()
+			=> false;
 
 		void OnStateChanged()
 		{
diff --git a/Engine/Source/Programs/UnrealToolbox/Program.cs b/Engine/Source/Programs/UnrealToolbox/Program.cs
--- a/Engine/Source/Programs/UnrealToolbox/Program.cs
+++ b/Engine/Source/Programs/UnrealToolbox/Program.cs
@@ -13,6 +13,7 @@ namespace UnrealToolbox
 	{
 		const string MutexName = "UnrealToolbox-Mutex";
 		public const string CloseEventName = "UnrealToolbox-Close";
+		public const string RefreshEventName = "UnrealToolbox-Refresh"; // Note: this event is set by InstallerCustomActions when agent is installed.
 		public const string SettingsEventName = "UnrealToolbox-Settings";
 
 		public static SelfUpdateState? Update { get; private set; }
@@ -48,6 +49,12 @@ namespace UnrealToolbox
 				closeEvent.Set();
 			}
 
+			using EventWaitHandle refreshEvent = new EventWaitHandle(false, EventResetMode.AutoReset, RefreshEventName);
+			if (args.Any(x => x.Equals("-Refresh", StringComparison.OrdinalIgnoreCase)))
+			{
+				refreshEvent.Set();
+			}
+
 			using EventWaitHandle settingsEvent = new EventWaitHandle(false, EventResetMode.AutoReset, SettingsEventName);
 			if (args.Any(x => x.Equals("-Settings", StringComparison.OrdinalIgnoreCase)))
 			{
diff --git a/Engine/Source/Programs/UnrealToolbox/SettingsWindow.axaml.cs b/Engine/Source/Programs/UnrealToolbox/SettingsWindow.axaml.cs
--- a/Engine/Source/Programs/UnrealToolbox/SettingsWindow.axaml.cs
+++ b/Engine/Source/Programs/UnrealToolbox/SettingsWindow.axaml.cs
@@ -26,21 +26,36 @@ namespace UnrealToolbox
 			InitializeComponent();
 
 			_serviceProvider = serviceProvider;
+			_navView.SelectionChanged += NavView_SelectionChanged;
+
+			Refresh();
+		}
 
-			foreach (ITrayAppPlugin plugin in _serviceProvider.GetServices<ITrayAppPlugin>())
+		public void Refresh()
+		{
+			if (_serviceProvider != null)
 			{
-				if(plugin.HasSettingsPage())
+				object? selectedItemContent = (_navView.SelectedItem as NavigationViewItem)?.Content;
+
+				_typeToPlugin.Clear();
+
+				_navView.MenuItems.Clear();
+				_navView.MenuItems.Add(new NavigationViewItem() { Content = "General", IconSource = new SymbolIconSource() { Symbol = Symbol.Settings }, Tag = typeof(GeneralSettingsPage).FullName });
+
+				foreach (ITrayAppPlugin plugin in _serviceProvider.GetServices<ITrayAppPlugin>())
 				{
-					string typeName = plugin.GetType().FullName!;
-					_navView.MenuItems.Add(new NavigationViewItem() { Content = plugin.Name, IconSource = plugin.Icon, Tag = typeName });
-					_typeToPlugin.Add(typeName, plugin);
+					if (plugin.HasSettingsPage())
+					{
+						string typeName = plugin.GetType().FullName!;
+						_navView.MenuItems.Add(new NavigationViewItem() { Content = plugin.Name, IconSource = plugin.Icon, Tag = typeName });
+						_typeToPlugin.Add(typeName, plugin);
+					}
 				}
-			}
 
-			_navView.SelectedItem = _navView.MenuItems[0];
-			_navView.SelectionChanged += NavView_SelectionChanged;
+				_navView.SelectedItem = _navView.MenuItems.FirstOrDefault(x => Object.Equals((x as NavigationViewItem)?.Content, selectedItemContent)) ?? _navView.MenuItems[0];
 
-			NavView_UpdateContent();
+				NavView_UpdateContent();
+			}
 		}
 
 		protected override void OnGotFocus(GotFocusEventArgs e)
