Unreal Toolbox: Activate settings window when opened.

diff --git a/Engine/Source/Programs/UnrealToolbox/App.axaml.cs b/Engine/Source/Programs/UnrealToolbox/App.axaml.cs
--- a/Engine/Source/Programs/UnrealToolbox/App.axaml.cs
+++ b/Engine/Source/Programs/UnrealToolbox/App.axaml.cs
@@ -257,14 +257,15 @@ namespace UnrealToolbox
 			if (_settingsWindow != null)
 			{
 				_settingsWindow.BringIntoView();
-				_settingsWindow.Show();
 			}
 			else
 			{
 				_settingsWindow = new SettingsWindow(_serviceProvider);
 				_settingsWindow.Closed += SettingsWindow_Closed;
-				_settingsWindow.Show();
 			}
+
+			_settingsWindow.Activate();
+			_settingsWindow.Show();
 		}
 
 		private void SettingsWindow_Closed(object? sender, EventArgs e)
