Unreal Toolbox: Fix hang on shutdown due to leak of HordeClientRef objects. Also prevent new instance starting when running with the -Close argument.

diff --git a/Engine/Source/Programs/UnrealToolbox/HordeClientProvider.cs b/Engine/Source/Programs/UnrealToolbox/HordeClientProvider.cs
--- a/Engine/Source/Programs/UnrealToolbox/HordeClientProvider.cs
+++ b/Engine/Source/Programs/UnrealToolbox/HordeClientProvider.cs
@@ -14,7 +14,7 @@ namespace UnrealToolbox
 	{
 		class HordeClientLifetime : IAsyncDisposable
 		{
-			readonly TaskCompletionSource _refZeroTcs = new TaskCompletionSource();
+			readonly TaskCompletionSource _refZeroTcs = new TaskCompletionSource(TaskCreationOptions.RunContinuationsAsynchronously);
 
 			readonly Task<IHordeClient?> _clientTask;
 			int _refCount;
diff --git a/Engine/Source/Programs/UnrealToolbox/Program.cs b/Engine/Source/Programs/UnrealToolbox/Program.cs
--- a/Engine/Source/Programs/UnrealToolbox/Program.cs
+++ b/Engine/Source/Programs/UnrealToolbox/Program.cs
@@ -47,12 +47,14 @@ namespace UnrealToolbox
 			if (args.Any(x => x.Equals("-Close", StringComparison.OrdinalIgnoreCase)))
 			{
 				closeEvent.Set();
+				return 0;
 			}
 
 			using EventWaitHandle refreshEvent = new EventWaitHandle(false, EventResetMode.AutoReset, RefreshEventName);
 			if (args.Any(x => x.Equals("-Refresh", StringComparison.OrdinalIgnoreCase)))
 			{
 				refreshEvent.Set();
+				return 0;
 			}
 
 			using EventWaitHandle settingsEvent = new EventWaitHandle(false, EventResetMode.AutoReset, SettingsEventName);
diff --git a/Engine/Source/Programs/UnrealToolbox/SelfUpdateService.cs b/Engine/Source/Programs/UnrealToolbox/SelfUpdateService.cs
--- a/Engine/Source/Programs/UnrealToolbox/SelfUpdateService.cs
+++ b/Engine/Source/Programs/UnrealToolbox/SelfUpdateService.cs
@@ -218,7 +218,7 @@ namespace UnrealToolbox
 				return;
 			}
 
-			IHordeClientRef? clientRef = await _hordeClientProvider.GetClientRefAsync().WaitAsync(cancellationToken);
+			using IHordeClientRef? clientRef = await _hordeClientProvider.GetClientRefAsync().WaitAsync(cancellationToken);
 			if (clientRef == null)
 			{
 				return;
