Horde: Keep notfications from creating start menu items, and changing notification aumid upon uninstall

diff --git a/Engine/Source/Programs/UnrealToolbox/Notifications/Notifications.cs b/Engine/Source/Programs/UnrealToolbox/Notifications/Notifications.cs
--- a/Engine/Source/Programs/UnrealToolbox/Notifications/Notifications.cs
+++ b/Engine/Source/Programs/UnrealToolbox/Notifications/Notifications.cs
@@ -9,6 +9,8 @@ using DesktopNotifications.FreeDesktop;
 using DesktopNotifications.Windows;
 using Microsoft.Extensions.Logging;
 using System;
+using System.Reflection;
+using System.Runtime.InteropServices;
 using UnrealToolbox;
 
 namespace UnrealToolbox
@@ -41,6 +43,9 @@ namespace UnrealToolbox
 		string? _lastTitle;
 		string? _lastBody;
 
+		[DllImport("shell32.dll", SetLastError = true)]
+		private static extern void SetCurrentProcessExplicitAppUserModelID([MarshalAs(UnmanagedType.LPWStr)] string appId);
+
 		public ToolboxNotificationManager(ILogger<ToolboxNotificationManager> logger)
 		{
 			_logger = logger;
@@ -50,7 +55,11 @@ namespace UnrealToolbox
 		{
 			if (Environment.OSVersion.Platform == PlatformID.Win32NT)
 			{
-				WindowsApplicationContext context = WindowsApplicationContext.FromCurrentProcess();
+				// WindowsApplicationContext.FromCurrentProcess() has side effects of creating start menu items, and changing the app user model id to the executing assembly, which can be dotnet.exe
+				// WindowsApplicationContext context = WindowsApplicationContext.FromCurrentProcess();
+				WindowsApplicationContext? context = Activator.CreateInstance(type: typeof(WindowsApplicationContext), bindingAttr: BindingFlags.Instance | BindingFlags.NonPublic, binder: null, args: new object[] { "UnrealToolbox", "UnrealToolbox" },culture: null) as WindowsApplicationContext;
+				SetCurrentProcessExplicitAppUserModelID("UnrealToolbox");
+
 				_platformManager = new WindowsNotificationManager(context);
 			}
 			else
