Unreal Toolbox: Always show the settings dialog on launch unless the -Quiet argument is passed on the command line. This makes a more sensible default than requiring a -Settings argument.

diff --git a/Engine/Source/Programs/UnrealToolbox/Program.cs b/Engine/Source/Programs/UnrealToolbox/Program.cs
--- a/Engine/Source/Programs/UnrealToolbox/Program.cs
+++ b/Engine/Source/Programs/UnrealToolbox/Program.cs
@@ -58,7 +58,7 @@ namespace UnrealToolbox
 			}
 
 			using EventWaitHandle settingsEvent = new EventWaitHandle(false, EventResetMode.AutoReset, SettingsEventName);
-			if (args.Any(x => x.Equals("-Settings", StringComparison.OrdinalIgnoreCase)))
+			if (!args.Any(x => x.Equals("-Quiet", StringComparison.OrdinalIgnoreCase)))
 			{
 				settingsEvent.Set();
 			}
