Horde: Include a Version.json file in Unreal Toolbox builds which can be used to detect upgrades/downgrades. When auto-updating from builds with this file, only upgrades will be allowed.

diff --git a/Engine/Source/Programs/Horde/BuildHorde.xml b/Engine/Source/Programs/Horde/BuildHorde.xml
--- a/Engine/Source/Programs/Horde/BuildHorde.xml
+++ b/Engine/Source/Programs/Horde/BuildHorde.xml
@@ -173,6 +173,7 @@
 		<Node Name="Build UnrealToolbox" Produces="#UnrealToolboxFiles">
 			<Delete Files="$(ServerToolsStagingDir_UnrealToolbox)/..."/>
 			<DotNet	Arguments="publish &quot;$(RootDir)/Engine/Source/Programs/UnrealToolbox/UnrealToolbox.csproj&quot; --output &quot;$(ServerToolsStagingDir_UnrealToolbox)&quot; --runtime win-x64 --self-contained $(VersionArguments)" DotNetPath="$(DotNetPath)" />
+			<WriteTextFile File="$(ServerToolsStagingDir_UnrealToolbox)/Version.json" Text="{ &quot;version&quot;: &quot;$(InformationalVersion)&quot; }"/>
 			<Sign Files="$(ServerToolsStagingDir_UnrealToolbox)/UnrealToolbox.exe;$(ServerToolsStagingDir_UnrealToolbox)/UnrealToolbox.dll;$(ServerToolsStagingDir_UnrealToolbox)/EpicGames.*.dll" If="'$(SignOutputs)' == 'true'"/>
 			<Tag Files="$(ServerToolsStagingDir_UnrealToolbox)/..." With="#UnrealToolboxFiles"/>
 		</Node>
@@ -290,14 +291,8 @@
 
 		</Node>
 		
-		<Node Name="Build Horde Agent Installer (Toolbox)" Requires="Build Horde Agent Installer" Produces="#HordeAgentMsiToolboxFiles">
-			<Copy Files="UnrealHordeAgent.msi" From="$(ServerToolsStagingDir_HordeAgentMsi)" To="$(ServerToolsStagingDir_HordeAgentMsiToolbox)"/>
-			<Copy Files="Toolbox.json" From="$(RootDir)/Engine/Source/Programs/Horde/HordeAgent" To="$(ServerToolsStagingDir_HordeAgentMsiToolbox)"/>
-			<Tag Files="$(ServerToolsStagingDir_HordeAgentMsiToolbox)/..." With="#HordeAgentMsiToolboxFiles"/>
-		</Node>
-
 		<ForEach Name="ServerArch" Values="$(ServerArchList)">
-			<Node Name="Stage Horde Server Installer ($(ServerArch))" Requires="#DashboardOutputFiles;#HordeCmdFiles;#HordeAgentFiles;#HordeAgentMsiFiles;#HordeAgentMsiToolboxFiles;#UnrealToolboxFiles;#UnrealToolboxMsiFiles;#P4VUtilsFiles;#UnrealGameSyncFiles;#UnrealGameSyncCmdFiles;#UnrealGameSyncMsiFiles;" Produces="#ServerWithToolsFiles_$(ServerArch)">
+			<Node Name="Stage Horde Server Installer ($(ServerArch))" Requires="#DashboardOutputFiles;#HordeCmdFiles;#HordeAgentFiles;#HordeAgentMsiFiles;#UnrealToolboxFiles;#UnrealToolboxMsiFiles;#P4VUtilsFiles;#UnrealGameSyncFiles;#UnrealGameSyncCmdFiles;#UnrealGameSyncMsiFiles;" Produces="#ServerWithToolsFiles_$(ServerArch)">
 		
 				<!-- Stage the server -->
 				<Delete Files="$(StagingDir)/Server/$(ServerArch)/..."/>
diff --git a/Engine/Source/Programs/UnrealToolbox/SelfUpdateService.cs b/Engine/Source/Programs/UnrealToolbox/SelfUpdateService.cs
--- a/Engine/Source/Programs/UnrealToolbox/SelfUpdateService.cs
+++ b/Engine/Source/Programs/UnrealToolbox/SelfUpdateService.cs
@@ -3,10 +3,11 @@
 using EpicGames.Core;
 using EpicGames.Horde.Storage.Nodes;
 using EpicGames.Horde.Tools;
-using Microsoft.AspNetCore.StaticFiles;
 using Microsoft.Extensions.Logging;
 using System.Diagnostics;
+using System.Diagnostics.CodeAnalysis;
 using System.Reflection;
+using System.Text.Json;
 
 namespace UnrealToolbox
 {
@@ -19,6 +20,7 @@ namespace UnrealToolbox
 			public string? UpdateVersion { get; set; }
 		}
 
+		readonly string? _currentVersion;
 		readonly DirectoryReference _baseDir;
 		readonly JsonConfig<JsonState> _config;
 		readonly string[] _args;
@@ -47,9 +49,10 @@ namespace UnrealToolbox
 			_args = args;
 
 			FileReference currentAssembly = new FileReference(Assembly.GetExecutingAssembly().Location);
+			_currentVersion = ReadVersion(currentAssembly.Directory);
 
 			_isLaunchApp = !currentAssembly.IsUnderDirectory(_baseDir);
-			if(_isLaunchApp && !String.Equals(_config.Current.LaunchApp, currentAssembly.FullName, StringComparison.OrdinalIgnoreCase))
+			if (_isLaunchApp && !String.Equals(_config.Current.LaunchApp, currentAssembly.FullName, StringComparison.OrdinalIgnoreCase))
 			{
 				UpdateState(x => x.LaunchApp = currentAssembly.FullName);
 			}
@@ -111,7 +114,7 @@ namespace UnrealToolbox
 				{
 					string fileName = Path.GetFileName(Assembly.GetExecutingAssembly().Location);
 					FileReference executable = FileReference.Combine(LatestDir, fileName);
-					if (FileReference.Exists(executable))
+					if (FileReference.Exists(executable) && ShouldLaunchLatest(LatestDir))
 					{
 						return Launch(executable.FullName, _args);
 					}
@@ -128,6 +131,69 @@ namespace UnrealToolbox
 			return false;
 		}
 
+		static bool TryParseVersion(string? version, [NotNullWhen(true)] out VersionNumber? versionNumber)
+		{
+			if (version == null)
+			{
+				versionNumber = null;
+				return false;
+			}
+
+			version = version.Replace("-PF-", ".", StringComparison.OrdinalIgnoreCase);
+			version = version.Replace("-", ".", StringComparison.Ordinal);
+
+			return VersionNumber.TryParse(version, out versionNumber);
+		}
+
+		bool ShouldLaunchLatest(DirectoryReference latestDir)
+		{
+			// If this build is not versioned, always launch latest
+			VersionNumber? currentVersionNumber;
+			if (!TryParseVersion(_currentVersion, out currentVersionNumber))
+			{
+				return true;
+			}
+
+			// If the latest build is versioned, make sure it's newer than the current version
+			VersionNumber? latestVersionNumber;
+			if (TryParseVersion(ReadVersion(latestDir), out latestVersionNumber))
+			{
+				return latestVersionNumber > currentVersionNumber;
+			}
+			else
+			{
+				return false;
+			}
+		}
+
+		class VersionFile
+		{
+			public string? Version { get; set; }
+		}
+
+		static string? ReadVersion(DirectoryReference dir)
+		{
+			FileReference file = FileReference.Combine(dir, "Version.json");
+			if (!FileReference.Exists(file))
+			{
+				return null;
+			}
+
+			string? version;
+			try
+			{
+				byte[] versionData = FileReference.ReadAllBytes(file);
+				VersionFile? versionFile = JsonSerializer.Deserialize<VersionFile>(versionData, new JsonSerializerOptions { PropertyNameCaseInsensitive = true });
+				version = versionFile?.Version;
+			}
+			catch
+			{
+				return null;
+			}
+
+			return version;
+		}
+
 		static bool Launch(string executable, string[] args)
 		{
 			using Process process = new Process();
