Unreal Toolbox: Only upgrade the toolbox app if both builds are versioned, and the latest build has a greater version number.

#rnx

diff --git a/Engine/Source/Programs/UnrealToolbox/SelfUpdateService.cs b/Engine/Source/Programs/UnrealToolbox/SelfUpdateService.cs
--- a/Engine/Source/Programs/UnrealToolbox/SelfUpdateService.cs
+++ b/Engine/Source/Programs/UnrealToolbox/SelfUpdateService.cs
@@ -147,23 +147,22 @@ namespace UnrealToolbox
 
 		bool ShouldLaunchLatest(DirectoryReference latestDir)
 		{
-			// If this build is not versioned, always launch latest
+			// If this build is not versioned, do not upgrade
 			VersionNumber? currentVersionNumber;
 			if (!TryParseVersion(_currentVersion, out currentVersionNumber))
 			{
-				return true;
+				return false;
 			}
 
-			// If the latest build is versioned, make sure it's newer than the current version
+			// If the other build is not versioned, do not upgrade
 			VersionNumber? latestVersionNumber;
-			if (TryParseVersion(ReadVersion(latestDir), out latestVersionNumber))
-			{
-				return latestVersionNumber > currentVersionNumber;
-			}
-			else
+			if (!TryParseVersion(ReadVersion(latestDir), out latestVersionNumber))
 			{
 				return false;
 			}
+
+			// Otherwise, only upgrade if the latest version is newer
+			return latestVersionNumber > currentVersionNumber;
 		}
 
 		class VersionFile
