Horde: Fix issue with TPS check on linux and avalonia dependency, part 1
#rnx

diff --git a/Engine/Source/Programs/UnrealToolbox/UnrealToolbox.csproj b/Engine/Source/Programs/UnrealToolbox/UnrealToolbox.csproj
--- a/Engine/Source/Programs/UnrealToolbox/UnrealToolbox.csproj
+++ b/Engine/Source/Programs/UnrealToolbox/UnrealToolbox.csproj
@@ -6,8 +6,6 @@
 
   <PropertyGroup>
     <OutputType>WinExe</OutputType>
-    <TargetFramework>net8.0-windows10.0.17763.0</TargetFramework>
-    <WindowsSdkPackageVersion>10.0.17763.41</WindowsSdkPackageVersion>
     <Nullable>enable</Nullable>
     <ImplicitUsings>enable</ImplicitUsings>
     <NeutralLanguage>en-US</NeutralLanguage>
@@ -17,7 +15,16 @@
     <NoWarn>IDE0005;CA1501</NoWarn>
     <ApplicationIcon>Resources/StatusNormal.ico</ApplicationIcon>
   </PropertyGroup>
+  
+  <PropertyGroup Condition="'$(EnableForTPS)'=='true'">
+     <TargetFramework>net8.0-windows</TargetFramework>
+   </PropertyGroup>
 
+  <PropertyGroup Condition="'$(EnableForTPS)'!='true'">
+     <TargetFramework>net8.0-windows10.0.17763.0</TargetFramework>
+     <WindowsSdkPackageVersion>10.0.17763.41</WindowsSdkPackageVersion>
+   </PropertyGroup>
+  
   <PropertyGroup Condition=" '$(Configuration)|$(Platform)' == 'Analyze|AnyCPU' ">
     <RunAnalyzersDuringBuild>True</RunAnalyzersDuringBuild>
     <TreatWarningsAsErrors>false</TreatWarningsAsErrors>
