Unreal Toolbox: Add an about page to the settings menu containing the version information.

diff --git a/Engine/Source/Programs/UnrealToolbox/AboutPage.axaml b/Engine/Source/Programs/UnrealToolbox/AboutPage.axaml
new file mode 100644
--- /dev/null
+++ b/Engine/Source/Programs/UnrealToolbox/AboutPage.axaml
@@ -0,0 +1,31 @@
+<UserControl xmlns="[REDACTED_PROFILE_URL]"
+			xmlns:x="http://schemas.microsoft.com/winfx/2006/xaml"
+			xmlns:d="http://schemas.microsoft.com/expression/blend/2008"
+			xmlns:mc="http://schemas.openxmlformats.org/markup-compatibility/2006"
+			xmlns:ui="using:FluentAvalonia.UI.Controls"
+			xmlns:local="using:UnrealToolbox"
+			mc:Ignorable="d" d:DesignWidth="800" d:DesignHeight="600"
+			x:Class="UnrealToolbox.AboutPage"
+			Background="{DynamicResource SolidBackgroundFillColorBaseBrush}">
+
+	<UserControl.Styles>
+		<StyleInclude Source="/Styles.axaml"/>
+	</UserControl.Styles>
+	
+	<ScrollViewer>
+		<StackPanel Margin="20" Spacing="5">
+
+			<TextBlock Theme="{StaticResource SubtitleTextBlockStyle}" VerticalAlignment="Center" Margin="0 0 0 0">About</TextBlock>
+
+			<TextBlock Classes="SettingHeader">Version</TextBlock>
+
+			<Border Classes="Card" Padding="16 6 16 2">
+				<StackPanel Spacing="10" Margin="0 10 20 15">
+					<TextBlock Classes="SettingHelp" Name="_versionText"/>
+				</StackPanel>
+			</Border>
+
+		</StackPanel>
+	</ScrollViewer>
+
+</UserControl>
diff --git a/Engine/Source/Programs/UnrealToolbox/AboutPage.axaml.cs b/Engine/Source/Programs/UnrealToolbox/AboutPage.axaml.cs
new file mode 100644
--- /dev/null
+++ b/Engine/Source/Programs/UnrealToolbox/AboutPage.axaml.cs
@@ -0,0 +1,20 @@
+// Copyright Epic Games, Inc. All Rights Reserved.
+
+using Avalonia;
+using Avalonia.Controls;
+using Avalonia.Markup.Xaml;
+using Microsoft.Extensions.DependencyInjection;
+
+namespace UnrealToolbox
+{
+	partial class AboutPage : UserControl
+	{
+		public AboutPage()
+		{
+			InitializeComponent();
+
+			string version = Program.Update?.CurrentVersionString ?? "No version information present.";
+			_versionText.Text = version;
+		}
+	}
+}
\ No newline at end of file
diff --git a/Engine/Source/Programs/UnrealToolbox/SelfUpdateService.cs b/Engine/Source/Programs/UnrealToolbox/SelfUpdateService.cs
--- a/Engine/Source/Programs/UnrealToolbox/SelfUpdateService.cs
+++ b/Engine/Source/Programs/UnrealToolbox/SelfUpdateService.cs
@@ -26,6 +26,9 @@ namespace UnrealToolbox
 		readonly string[] _args;
 		readonly bool _isLaunchApp;
 
+		public string? CurrentVersionString
+			=> _currentVersion;
+
 		public string? LatestVersion
 			=> _config.Current.LatestVersion;
 
diff --git a/Engine/Source/Programs/UnrealToolbox/SettingsWindow.axaml.cs b/Engine/Source/Programs/UnrealToolbox/SettingsWindow.axaml.cs
--- a/Engine/Source/Programs/UnrealToolbox/SettingsWindow.axaml.cs
+++ b/Engine/Source/Programs/UnrealToolbox/SettingsWindow.axaml.cs
@@ -54,6 +54,8 @@ namespace UnrealToolbox
 
 				_navView.SelectedItem = _navView.MenuItems.FirstOrDefault(x => Object.Equals((x as NavigationViewItem)?.Content, selectedItemContent)) ?? _navView.MenuItems[0];
 
+				_navView.FooterMenuItems.Add(new NavigationViewItem() { Content = "About", IconSource = new SymbolIconSource() { Symbol = Symbol.Help } });
+
 				NavView_UpdateContent();
 			}
 		}
@@ -90,6 +92,10 @@ namespace UnrealToolbox
 				{
 					_navView.Content = plugin.CreateSettingsPage(context);
 				}
+				else if (nvi == _navView.FooterMenuItems.FirstOrDefault())
+				{
+					_navView.Content = new AboutPage();
+				}
 				else
 				{
 					_navView.Content = new GeneralSettingsPage(context);
