Unreal Toolbox: Do not update agent plugin status when an install/uninstall operation is in progress.

diff --git a/Engine/Source/Programs/UnrealToolbox/Plugins/HordeAgent/HordeAgentPlugin.cs b/Engine/Source/Programs/UnrealToolbox/Plugins/HordeAgent/HordeAgentPlugin.cs
--- a/Engine/Source/Programs/UnrealToolbox/Plugins/HordeAgent/HordeAgentPlugin.cs
+++ b/Engine/Source/Programs/UnrealToolbox/Plugins/HordeAgent/HordeAgentPlugin.cs
@@ -19,6 +19,7 @@ namespace UnrealToolbox.Plugins.HordeAgent
 		record struct IdleStat(string Name, long Value, long MinValue);
 
 		readonly ITrayAppHost _host;
+		readonly IToolCatalog _toolCatalog;
 		readonly ILogger _logger;
 
 		readonly BackgroundTask _clientTask;
@@ -33,6 +34,9 @@ namespace UnrealToolbox.Plugins.HordeAgent
 
 		AgentSettingsMessage? _agentSettings;
 
+		TrayAppPluginStatus? _status;
+		TrayAppPluginStatus? _reportStatus; // Updated with _status when nothing is currently being installed
+
 		public string Name => "Horde Agent";
 
 		public IconSource Icon => new SymbolIconSource() { Symbol = Symbol.People };
@@ -69,9 +73,10 @@ namespace UnrealToolbox.Plugins.HordeAgent
 			}
 		}
 
-		public HordeAgentPlugin(ITrayAppHost host, ILogger<HordeAgentPlugin> logger)
+		public HordeAgentPlugin(ITrayAppHost host, ToolCatalog toolCatalog, ILogger<HordeAgentPlugin> logger)
 		{
 			_host = host;
+			_toolCatalog = toolCatalog;
 			_logger = logger;
 
 			DirectoryReference? settingsRoot = DirectoryReference.GetSpecialFolder(Environment.SpecialFolder.LocalApplicationData);
@@ -217,8 +222,6 @@ namespace UnrealToolbox.Plugins.HordeAgent
 			}
 		}
 
-		TrayAppPluginStatus? _status;
-
 		void SetStatus(AgentStatusMessage status)
 		{
 			if (!IsEnabled)
@@ -249,7 +252,13 @@ namespace UnrealToolbox.Plugins.HordeAgent
 		}
 
 		public TrayAppPluginStatus GetStatus()
-			=> _status ?? TrayAppPluginStatus.Default;
+		{
+			if (!_toolCatalog.Items.Any(x => x.Pending != null && !x.Pending.Failed))
+			{
+				_reportStatus = _status;
+			}
+			return _reportStatus ?? TrayAppPluginStatus.Default;
+		}
 
 		async Task StatusTaskAsync(CancellationToken cancellationToken)
 		{
