Unreal Toolbox: Missing csproj file.

#rnx

diff --git a/Engine/Source/Programs/UnrealToolbox/UnrealToolbox.csproj b/Engine/Source/Programs/UnrealToolbox/UnrealToolbox.csproj
--- a/Engine/Source/Programs/UnrealToolbox/UnrealToolbox.csproj
+++ b/Engine/Source/Programs/UnrealToolbox/UnrealToolbox.csproj
@@ -59,6 +59,9 @@
     <Compile Update="SettingsWindow.axaml.cs">
       <DependentUpon>SettingsWindow.axaml</DependentUpon>
     </Compile>
+    <Compile Update="AboutPage.axaml.cs">
+      <DependentUpon>AboutPage.axaml</DependentUpon>
+    </Compile>
   </ItemGroup>
 
   <ItemGroup>
