Unreal Toolbox: Prevent multiple "about" menu items being added to toolbox when refreshing tabs.

diff --git a/Engine/Source/Programs/UnrealToolbox/SettingsWindow.axaml.cs b/Engine/Source/Programs/UnrealToolbox/SettingsWindow.axaml.cs
--- a/Engine/Source/Programs/UnrealToolbox/SettingsWindow.axaml.cs
+++ b/Engine/Source/Programs/UnrealToolbox/SettingsWindow.axaml.cs
@@ -52,10 +52,14 @@ namespace UnrealToolbox
 					}
 				}
 
-				_navView.SelectedItem = _navView.MenuItems.FirstOrDefault(x => Object.Equals((x as NavigationViewItem)?.Content, selectedItemContent)) ?? _navView.MenuItems[0];
-
+				_navView.FooterMenuItems.Clear();
 				_navView.FooterMenuItems.Add(new NavigationViewItem() { Content = "About", IconSource = new SymbolIconSource() { Symbol = Symbol.Help } });
 
+				_navView.SelectedItem = 
+					_navView.MenuItems.FirstOrDefault(x => Object.Equals((x as NavigationViewItem)?.Content, selectedItemContent))
+					?? _navView.FooterMenuItems.FirstOrDefault(x => Object.Equals((x as NavigationViewItem)?.Content, selectedItemContent))
+					?? _navView.MenuItems[0];
+
 				NavView_UpdateContent();
 			}
 		}
