Unreal Toolbox: Rename "Connected" status to "Authenticated". This status does not reflect an actual connection, but just the state of the OIDC token.

diff --git a/Engine/Source/Programs/UnrealToolbox/GeneralSettingsViewModel.cs b/Engine/Source/Programs/UnrealToolbox/GeneralSettingsViewModel.cs
--- a/Engine/Source/Programs/UnrealToolbox/GeneralSettingsViewModel.cs
+++ b/Engine/Source/Programs/UnrealToolbox/GeneralSettingsViewModel.cs
@@ -63,7 +63,7 @@ namespace UnrealToolbox
 			else
 			{
 				ServerUrl = hordeClient.Client.ServerUrl.ToString();
-				ServerStatus = hordeClient.Client.HasValidAccessToken() ? "Connected" : "Session Expired";
+				ServerStatus = hordeClient.Client.HasValidAccessToken() ? "Authenticated" : "Session Expired";
 				IsLoginEnabled = !hordeClient.Client.HasValidAccessToken();
 			}
 		}
