Unreal Toolbox: Write log file to disk (in AppData\Epic Games\Unreal Toolbox) and provide a clickable link to open the log when installation fails.

diff --git a/Engine/Source/Programs/UnrealToolbox/GeneralSettingsPage.axaml b/Engine/Source/Programs/UnrealToolbox/GeneralSettingsPage.axaml
new file mode 100644
--- /dev/null
+++ b/Engine/Source/Programs/UnrealToolbox/GeneralSettingsPage.axaml
@@ -0,0 +1,76 @@
+<UserControl xmlns="[REDACTED_PROFILE_URL]"
+			xmlns:x="http://schemas.microsoft.com/winfx/2006/xaml"
+			xmlns:d="http://schemas.microsoft.com/expression/blend/2008"
+			xmlns:mc="http://schemas.openxmlformats.org/markup-compatibility/2006"
+			xmlns:ui="using:FluentAvalonia.UI.Controls"
+			xmlns:local="using:UnrealToolbox"
+			mc:Ignorable="d" d:DesignWidth="800" d:DesignHeight="600"
+			x:Class="UnrealToolbox.GeneralSettingsPage"
+			x:DataType="local:GeneralSettingsViewModel"
+			Background="{DynamicResource SolidBackgroundFillColorBaseBrush}">
+
+	<UserControl.Styles>
+		<StyleInclude Source="/Styles.axaml"/>
+	</UserControl.Styles>
+	
+	<ScrollViewer>
+		<StackPanel Margin="20" Spacing="5">
+
+			<TextBlock Theme="{StaticResource SubtitleTextBlockStyle}" VerticalAlignment="Center" Margin="0 0 0 0">General</TextBlock>
+
+			<TextBlock Classes="SettingHeader">Horde Server</TextBlock>
+
+			<Border Classes="Card">
+				<Grid ColumnDefinitions="*,Auto,Auto" RowDefinitions="Auto,Auto">
+					<StackPanel Grid.Column="0" Margin="0">
+						<HyperlinkButton Grid.Column="0" Grid.Row="0" Padding="0" FontSize="14" Content="{Binding ServerUrl}" Command="{Binding OpenServer}"/>
+						<TextBlock Grid.Column="0" Grid.Row="1" Padding="1" FontSize="12" Text="{Binding ServerStatus}"/>
+					</StackPanel>
+					<Button Grid.Column="1" Command="{Binding LoginAsync}" IsEnabled="{Binding IsLoginEnabled}" Margin="0 0 5 0">Connect</Button>
+					<Button Grid.Column="2" Command="{Binding ConfigureServerAsync}">Modify...</Button>
+				</Grid>
+			</Border>
+
+			<StackPanel IsVisible="{Binding ToolCatalog.ShowTools}">
+				<TextBlock Classes="SettingHeader">Updates</TextBlock>
+				
+				<Border Classes="Card" Padding="16 6 16 2">
+					<StackPanel Spacing="10" Margin="0 0 0 5">
+						<Grid ColumnDefinitions="*,Auto">
+							<StackPanel Grid.Column="0" Margin="0 10 20 10">
+								<TextBlock Classes="SettingCaption">Automatically update tools to latest version</TextBlock>
+								<TextBlock Classes="SettingHelp">Installed tools will be updated automatically when new versions become available.</TextBlock>
+							</StackPanel>
+							<ToggleSwitch Grid.Column="1" VerticalAlignment="Center" OnContent="" OffContent="" IsChecked="{Binding ToolCatalog.AutoUpdate}"/>
+						</Grid>
+					</StackPanel>
+				</Border>
+
+				<TextBlock Classes="SettingHeader">Tools</TextBlock>
+
+				<Border Classes="Card">
+					<ListBox x:Name="ToolsList" BorderBrush="Black" x:DataType="local:GeneralSettingsViewModel" ItemsSource="{Binding ToolCatalog.Items}">
+						<ListBox.ItemTemplate>
+							<DataTemplate DataType="local:ToolCatalogItemViewModel">
+								<Grid ColumnDefinitions="*,Auto,Auto,Auto,Auto,Auto,Auto">
+									<Grid RowDefinitions="Auto,Auto" Margin="0 10">
+										<TextBlock Text="{Binding Name}" VerticalAlignment="Center" FontSize="12" FontWeight="SemiBold" Grid.Row="0"/>
+										<TextBlock Text="{Binding Description}" VerticalAlignment="Center" FontSize="12" FontWeight="Light" Grid.Row="1"/>
+									</Grid>
+									<ui:ProgressRing Width="16" Height="16" IsVisible="{Binding IsBusy}" VerticalAlignment="Center" Grid.Row="0" Grid.Column="1"/>
+									<TextBlock Text="{Binding Status}" IsVisible="{Binding ShowStatusText}" VerticalAlignment="Center" Padding="5" FontSize="12" FontWeight="Light" Grid.Row="0" Grid.Column="2" Margin="4 0 4 0"/>
+									<HyperlinkButton Content="{Binding Status}" IsVisible="{Binding ShowStatusLink}" Command="{Binding ShowLog}" FontSize="12" FontWeight="Light" Grid.Column="3" Margin="4 0 4 0"/>
+									<Button Command="{Binding Add}" IsVisible="{Binding ShowAdd}" Grid.Column="4">Add</Button>
+									<Button Command="{Binding Update}" IsVisible="{Binding ShowUpdate}" Grid.Column="5">Update</Button>
+									<Button Command="{Binding Remove}" IsVisible="{Binding ShowRemove}" Grid.Column="6">Remove</Button>
+								</Grid>
+							</DataTemplate>
+						</ListBox.ItemTemplate>
+					</ListBox>
+				</Border>
+			</StackPanel>
+			
+		</StackPanel>
+	</ScrollViewer>
+
+</UserControl>
diff --git a/Engine/Source/Programs/UnrealToolbox/IToolCatalogItem.cs b/Engine/Source/Programs/UnrealToolbox/IToolCatalogItem.cs
--- a/Engine/Source/Programs/UnrealToolbox/IToolCatalogItem.cs
+++ b/Engine/Source/Programs/UnrealToolbox/IToolCatalogItem.cs
@@ -83,5 +83,6 @@ namespace UnrealToolbox
 	/// <param name="Failed">Whether the installation failed</param>
 	/// <param name="Message">Current status message</param>
 	/// <param name="Deployment">Information about the pending deployment</param>
-	public record class PendingToolDeploymentInfo(bool Failed, string? Message, ToolDeploymentInfo? Deployment);
+	/// <param name="ShowLogLink">Whether to show a link to the log</param>
+	public record class PendingToolDeploymentInfo(bool Failed, string? Message, ToolDeploymentInfo? Deployment, bool ShowLogLink = false);
 }
diff --git a/Engine/Source/Programs/UnrealToolbox/Program.cs b/Engine/Source/Programs/UnrealToolbox/Program.cs
--- a/Engine/Source/Programs/UnrealToolbox/Program.cs
+++ b/Engine/Source/Programs/UnrealToolbox/Program.cs
@@ -18,6 +18,8 @@ namespace UnrealToolbox
 
 		public static SelfUpdateState? Update { get; private set; }
 
+		public static FileReference? LogFile { get; } = GetLogFile();
+
 		[STAThread]
 		public static int Main(string[] args)
 		{
@@ -69,6 +71,11 @@ namespace UnrealToolbox
 				return 1;
 			}
 
+			if (LogFile != null)
+			{
+				Log.AddFileWriter("Default", LogFile);
+			}
+
 			AppBuilder builder = BuildAvaloniaApp();
 			try
 			{
@@ -92,6 +99,16 @@ namespace UnrealToolbox
 				.UsePlatformDetect()
 				.LogToTrace();
 		}
+
+		static FileReference? GetLogFile()
+		{
+			DirectoryReference? appDataDir = DirectoryReference.GetSpecialFolder(Environment.SpecialFolder.LocalApplicationData);
+			if (appDataDir == null)
+			{
+				return null;
+			}
+			return FileReference.Combine(appDataDir, "Epic Games", "Unreal Toolbox", "Log.txt");
+		}
 	}
 
 	class SingleInstanceMutex : IDisposable
diff --git a/Engine/Source/Programs/UnrealToolbox/ToolCatalog.cs b/Engine/Source/Programs/UnrealToolbox/ToolCatalog.cs
--- a/Engine/Source/Programs/UnrealToolbox/ToolCatalog.cs
+++ b/Engine/Source/Programs/UnrealToolbox/ToolCatalog.cs
@@ -8,6 +8,7 @@ using Microsoft.Extensions.Logging;
 using Microsoft.Win32;
 using System.Diagnostics;
 using System.Runtime.InteropServices;
+using System.Text;
 using System.Text.Json;
 
 #nullable enable
@@ -485,7 +486,7 @@ namespace UnrealToolbox
 			catch (Exception ex)
 			{
 				_logger.LogError(ex, "Error updating deployment: {Message}", ex.Message);
-				UpdateStateAndNotify(item, state => state with { Pending = new PendingToolDeploymentInfo(true, $"Error updating. See log.", deploymentInfo) });
+				UpdateStateAndNotify(item, state => state with { Pending = new PendingToolDeploymentInfo(true, $"Error updating. View Log for info.", deploymentInfo, ShowLogLink: true) });
 			}
 		}
 
@@ -493,7 +494,7 @@ namespace UnrealToolbox
 		{
 			ToolConfig toolConfig = LoadToolConfig(toolDir);
 
-			await UninstallAsync(item, cancellationToken);
+			await UninstallInternalAsync(item, cancellationToken);
 
 			if (toolConfig.InstallCommand != null)
 			{
@@ -546,6 +547,12 @@ namespace UnrealToolbox
 		}
 
 		async Task UninstallAsync(Item item, CancellationToken cancellationToken)
+		{
+			await UninstallInternalAsync(item, cancellationToken);
+			UpdateStateAndNotify(item, state => state with { Pending = null, Current = null });
+		}
+
+		async Task UninstallInternalAsync(Item item, CancellationToken cancellationToken)
 		{
 			CurrentToolDeploymentInfo? current = item._state.Current;
 			if (current != null)
@@ -555,7 +562,7 @@ namespace UnrealToolbox
 					PendingToolDeploymentInfo? pending = new PendingToolDeploymentInfo(false, $"Removing {current.Version}", null);
 					UpdateStateAndNotify(item, state => state with { Pending = pending, Current = null });
 
-					await RunCommandAsync(current.Id.ToString(), current.Config.UninstallCommand, current.Dir.FullName, cancellationToken);
+					await RunCommandAsync(item.Id.ToString(), current.Config.UninstallCommand, current.Dir.FullName, cancellationToken);
 				}
 				else if (OperatingSystem.IsWindows() && !String.IsNullOrEmpty(item._state.MsiProductId) && IsMsiInstalled(item._state.MsiProductId))
 				{
@@ -565,7 +572,7 @@ namespace UnrealToolbox
 					ToolCommand uninstallCommand = new ToolCommand { FileName = "msiexec.exe" };
 					uninstallCommand.Arguments = new List<string> { "/x", $"{{{item._state.MsiProductId}}}" };
 
-					int exitCode = await RunCommandAsync(current.Id.ToString(), uninstallCommand, current.Dir.FullName, cancellationToken);
+					int exitCode = await RunCommandAsync(item.Id.ToString(), uninstallCommand, current.Dir.FullName, cancellationToken);
 					if (exitCode == 1602)
 					{
 						// User cancelled (https://learn.microsoft.com/en-us/windows/win32/msi/error-codes)
@@ -580,8 +587,6 @@ namespace UnrealToolbox
 					}
 				}
 			}
-
-			UpdateStateAndNotify(item, state => state with { Pending = null, Current = null });
 		}
 
 		async Task<int> RunCommandAsync(string toolName, ToolCommand command, string workingDir, CancellationToken cancellationToken)
@@ -598,6 +603,8 @@ namespace UnrealToolbox
 				fileName = command.FileName;
 			}
 
+			_logger.LogInformation("{ToolName}> Running {FileName} {Arguments}", toolName, CommandLineArguments.Quote(fileName), arguments);
+
 			using (ManagedProcessGroup newProcessGroup = new ManagedProcessGroup())
 			using (ManagedProcess newProcess = new ManagedProcess(newProcessGroup, fileName, arguments, workingDir, null, null, ProcessPriorityClass.Normal))
 			{
@@ -619,11 +626,47 @@ namespace UnrealToolbox
 		{
 			lock (_workerThreadLockObject)
 			{
+				string prevState = GetItemStateMessage(item._state);
 				item._state = update(item._state);
+
+				string nextState = GetItemStateMessage(item._state);
+				if (!String.Equals(prevState, nextState, StringComparison.Ordinal))
+				{
+					_logger.LogInformation("{ToolName}> {State}", item.Id, nextState);
+				}
 			}
 			PostMainThreadUpdate();
 		}
 
+		static string GetItemStateMessage(ItemState state)
+		{
+			List<string> messages = new List<string>();
+
+			if (state.Current != null)
+			{
+				messages.Add($"Current: {state.Current.Version} ({state.Current.Id})");
+			}
+
+			if (state.Pending != null)
+			{
+				if (state.Pending.Failed)
+				{
+					messages.Add($"Pending: {state.Pending.Message} (Failed)");
+				}
+				else if (state.Pending.Deployment != null)
+				{
+					messages.Add($"Pending: {state.Pending.Message} ({state.Pending.Deployment.Id}@{state.Pending.Deployment.Version})");
+				}
+			}
+
+			if (messages.Count == 0)
+			{
+				messages.Add("Not installed");
+			}
+
+			return String.Join(", ", messages);
+		}
+
 		#region Cleanup
 
 		async Task CleanupAsync(CancellationToken cancellationToken)
diff --git a/Engine/Source/Programs/UnrealToolbox/ToolCatalogViewModel.cs b/Engine/Source/Programs/UnrealToolbox/ToolCatalogViewModel.cs
--- a/Engine/Source/Programs/UnrealToolbox/ToolCatalogViewModel.cs
+++ b/Engine/Source/Programs/UnrealToolbox/ToolCatalogViewModel.cs
@@ -2,7 +2,9 @@
 
 using System.Collections.ObjectModel;
 using System.ComponentModel;
+using System.Diagnostics;
 using CommunityToolkit.Mvvm.ComponentModel;
+using EpicGames.Core;
 
 namespace UnrealToolbox
 {
@@ -112,6 +114,12 @@ namespace UnrealToolbox
 		[ObservableProperty]
 		string _status = String.Empty;
 
+		[ObservableProperty]
+		bool _showStatusLink;
+
+		[ObservableProperty]
+		bool _showStatusText;
+
 		[ObservableProperty]
 		bool _showAdd;
 
@@ -160,11 +168,24 @@ namespace UnrealToolbox
 				IsBusy = false;
 			}
 
+			ShowStatusLink = _item.Pending != null && _item.Pending.ShowLogLink;
+			ShowStatusText = !String.IsNullOrEmpty(Status) && !ShowStatusLink;
 			ShowAdd = _item.Latest != null && _item.Current == null;
 			ShowUpdate = _item.Latest != null && _item.Current != null && _item.Latest.Id != _item.Current.Id;
 			ShowRemove = _item.Current != null;
 		}
 
+		public void ShowLog()
+		{
+			if (OperatingSystem.IsWindows() && Program.LogFile != null)
+			{
+				ProcessStartInfo info = new ProcessStartInfo();
+				info.FileName = Program.LogFile.FullName;
+				info.UseShellExecute = true;
+				Process.Start(info);
+			}
+		}
+
 		public void Add()
 			=> _item.Install();
 
