Unreal Toolbox: Remove the minimized flag when showing an existing settings dialog.

#rnx

diff --git a/Engine/Source/Programs/UnrealToolbox/App.axaml.cs b/Engine/Source/Programs/UnrealToolbox/App.axaml.cs
--- a/Engine/Source/Programs/UnrealToolbox/App.axaml.cs
+++ b/Engine/Source/Programs/UnrealToolbox/App.axaml.cs
@@ -257,6 +257,7 @@ namespace UnrealToolbox
 		{
 			if (_settingsWindow != null)
 			{
+				_settingsWindow.WindowState &= ~WindowState.Minimized;
 				_settingsWindow.BringIntoView();
 			}
 			else
