Unreal Toolbox: Work around Avalonia not bringing the settings window to foreground when activated from notification icon.

#rnx

diff --git a/Engine/Source/Programs/UnrealToolbox/App.axaml.cs b/Engine/Source/Programs/UnrealToolbox/App.axaml.cs
--- a/Engine/Source/Programs/UnrealToolbox/App.axaml.cs
+++ b/Engine/Source/Programs/UnrealToolbox/App.axaml.cs
@@ -12,6 +12,7 @@ using UnrealToolbox.Plugins.HordeAgent;
 using UnrealToolbox.Plugins.HordeProxy;
 using Microsoft.Extensions.DependencyInjection;
 using DesktopNotifications;
+using System.Runtime.InteropServices;
 
 namespace UnrealToolbox
 {
@@ -253,6 +254,9 @@ namespace UnrealToolbox
 			OpenSettings();
 		}
 
+		[DllImport("user32.dll")]
+		static extern uint SetForegroundWindow(nint hWnd);
+
 		private void OpenSettings()
 		{
 			if (_settingsWindow != null)
@@ -268,6 +272,16 @@ namespace UnrealToolbox
 
 			_settingsWindow.Activate();
 			_settingsWindow.Show();
+
+			if (OperatingSystem.IsWindows())
+			{
+				// Avalonia doesn't seem to activate the window despite the call above, so fall back to pinvoke
+				nint? handle = _settingsWindow.TryGetPlatformHandle()?.Handle;
+				if (handle.HasValue)
+				{
+					SetForegroundWindow(handle.Value);
+				}
+			}
 		}
 
 		private void SettingsWindow_Closed(object? sender, EventArgs e)
