Unreal Toolbox: If the application cannot auto-update due to the settings window being open, trigger it the moment the settings window is closed.

#rnx

diff --git a/Engine/Source/Programs/UnrealToolbox/App.axaml.cs b/Engine/Source/Programs/UnrealToolbox/App.axaml.cs
--- a/Engine/Source/Programs/UnrealToolbox/App.axaml.cs
+++ b/Engine/Source/Programs/UnrealToolbox/App.axaml.cs
@@ -32,6 +32,7 @@ namespace UnrealToolbox
 
 		Thread? _settingsThread;
 		ManualResetEvent? _settingsThreadStop;
+		bool _shutdownOnClose;
 
 		/// <summary>
 		/// Constructor
@@ -141,6 +142,10 @@ namespace UnrealToolbox
 			{
 				((IClassicDesktopStyleApplicationLifetime)ApplicationLifetime!).Shutdown();
 			}
+			else
+			{
+				_shutdownOnClose = true;
+			}
 		}
 
 		private void UpdateMenu()
@@ -287,6 +292,11 @@ namespace UnrealToolbox
 		private void SettingsWindow_Closed(object? sender, EventArgs e)
 		{
 			_settingsWindow = null;
+
+			if (_shutdownOnClose)
+			{
+				((IClassicDesktopStyleApplicationLifetime)ApplicationLifetime!).Shutdown();
+			}
 		}
 
 		private void TrayIcon_Exit(object? sender, EventArgs e)
