UnrealToolbox: Disable unsafe code warnings

diff --git a/Engine/Source/Programs/UnrealToolbox/UnrealToolbox.csproj b/Engine/Source/Programs/UnrealToolbox/UnrealToolbox.csproj
--- a/Engine/Source/Programs/UnrealToolbox/UnrealToolbox.csproj
+++ b/Engine/Source/Programs/UnrealToolbox/UnrealToolbox.csproj
@@ -10,9 +10,10 @@
     <ImplicitUsings>enable</ImplicitUsings>
     <NeutralLanguage>en-US</NeutralLanguage>
     <Configurations>Debug;Release;Analyze</Configurations>
-	<AssemblyName>UnrealToolbox</AssemblyName>
+    <AssemblyName>UnrealToolbox</AssemblyName>
     <AvaloniaUseCompiledBindingsByDefault>true</AvaloniaUseCompiledBindingsByDefault>
     <NoWarn>IDE0005;CA1501</NoWarn>
+    <AllowUnsafeBlocks>true</AllowUnsafeBlocks>
     <ApplicationIcon>Resources/StatusNormal.ico</ApplicationIcon>
   </PropertyGroup>
   
