Unreal Toolbox: Add a semaphore around running MSI installers, to prevent multiple being started at once.

#rnx

diff --git a/Engine/Source/Programs/UnrealToolbox/ToolCatalog.cs b/Engine/Source/Programs/UnrealToolbox/ToolCatalog.cs
--- a/Engine/Source/Programs/UnrealToolbox/ToolCatalog.cs
+++ b/Engine/Source/Programs/UnrealToolbox/ToolCatalog.cs
@@ -203,6 +203,7 @@ namespace UnrealToolbox
 		readonly BackgroundTask _updateTask;
 		readonly AsyncEvent _cleanupEvent;
 		readonly BackgroundTask _cleanupTask;
+		readonly SemaphoreSlim _msiSemaphore;
 
 		readonly SynchronizationContext? _mainThreadSynchronizationContext;
 
@@ -247,6 +248,7 @@ namespace UnrealToolbox
 			_cleanupEvent = new AsyncEvent();
 			_cleanupTask = new BackgroundTask(CleanupAsync);
 			_mainThreadSynchronizationContext = SynchronizationContext.Current;
+			_msiSemaphore = new SemaphoreSlim(1);
 
 			_items = LoadState();
 			_mainThreadItems = null!;
@@ -337,6 +339,7 @@ namespace UnrealToolbox
 			await StopAsync(CancellationToken.None);
 			await _updateTask.DisposeAsync();
 			await _cleanupTask.DisposeAsync();
+			_msiSemaphore.Dispose();
 		}
 
 		void OnServerChange()
@@ -511,25 +514,42 @@ namespace UnrealToolbox
 			}
 			else if (OperatingSystem.IsWindows() && !String.IsNullOrEmpty(item._state.MsiProductId))
 			{
-				PendingToolDeploymentInfo pending = new PendingToolDeploymentInfo(false, "Running installer...", deploymentInfo);
-				UpdateStateAndNotify(item, state => state with { Pending = pending });
-
-				FileReference? msiFile = DirectoryReference.EnumerateFiles(toolDir, "*.msi").FirstOrDefault();
-				if (msiFile == null)
+				// Wait to run the installer
+				if (!await _msiSemaphore.WaitAsync(0, cancellationToken))
 				{
-					pending = new PendingToolDeploymentInfo(true, $"MSI file not found", deploymentInfo);
+					PendingToolDeploymentInfo pending = new PendingToolDeploymentInfo(false, "Waiting to install...", deploymentInfo);
 					UpdateStateAndNotify(item, state => state with { Pending = pending });
-					return;
+
+					await _msiSemaphore.WaitAsync(cancellationToken);
 				}
 
-				ToolCommand installCommand = new ToolCommand { FileName = "msiexec.exe" };
-				installCommand.Arguments = new List<string> { "/I", msiFile.FullName };
+				// Keep the lock until the install has finished
+				try
+				{
+					PendingToolDeploymentInfo pending = new PendingToolDeploymentInfo(false, "Running installer...", deploymentInfo);
+					UpdateStateAndNotify(item, state => state with { Pending = pending });
+
+					FileReference? msiFile = DirectoryReference.EnumerateFiles(toolDir, "*.msi").FirstOrDefault();
+					if (msiFile == null)
+					{
+						pending = new PendingToolDeploymentInfo(true, $"MSI file not found", deploymentInfo);
+						UpdateStateAndNotify(item, state => state with { Pending = pending });
+						return;
+					}
+
+					ToolCommand installCommand = new ToolCommand { FileName = "msiexec.exe" };
+					installCommand.Arguments = new List<string> { "/I", msiFile.FullName };
 
-				int exitCode = await RunCommandAsync(item.Id.ToString(), installCommand, toolDir.FullName, cancellationToken);
-				if (exitCode != 0 || !IsMsiInstalled(item._state.MsiProductId))
+					int exitCode = await RunCommandAsync(item.Id.ToString(), installCommand, toolDir.FullName, cancellationToken);
+					if (exitCode != 0 || !IsMsiInstalled(item._state.MsiProductId))
+					{
+						UpdateStateAndNotify(item, state => state with { Pending = null });
+						return;
+					}
+				}
+				finally
 				{
-					UpdateStateAndNotify(item, state => state with { Pending = null });
-					return;
+					_msiSemaphore.Release();
 				}
 			}
 
@@ -566,24 +586,43 @@ namespace UnrealToolbox
 				}
 				else if (OperatingSystem.IsWindows() && !String.IsNullOrEmpty(item._state.MsiProductId) && IsMsiInstalled(item._state.MsiProductId))
 				{
-					PendingToolDeploymentInfo pending = new PendingToolDeploymentInfo(false, $"Running installer...", current);
-					UpdateStateAndNotify(item, state => state with { Pending = pending });
+					PendingToolDeploymentInfo pending;
 
-					ToolCommand uninstallCommand = new ToolCommand { FileName = "msiexec.exe" };
-					uninstallCommand.Arguments = new List<string> { "/x", $"{{{item._state.MsiProductId}}}" };
-
-					int exitCode = await RunCommandAsync(item.Id.ToString(), uninstallCommand, current.Dir.FullName, cancellationToken);
-					if (exitCode == 1602)
+					// Wait to run the installer
+					if (!await _msiSemaphore.WaitAsync(0, cancellationToken))
 					{
-						// User cancelled (https://learn.microsoft.com/en-us/windows/win32/msi/error-codes)
-						UpdateStateAndNotify(item, state => state with { Pending = null });
-						return;
+						pending = new PendingToolDeploymentInfo(false, "Waiting to run installer...", current);
+						UpdateStateAndNotify(item, state => state with { Pending = pending });
+
+						await _msiSemaphore.WaitAsync(cancellationToken);
 					}
-					else if (exitCode != 0 || IsMsiInstalled(item._state.MsiProductId))
+
+					// Run the installer
+					try
 					{
-						pending = new PendingToolDeploymentInfo(true, $"Uninstall failed ({exitCode}).", current);
+						pending = new PendingToolDeploymentInfo(false, $"Running installer...", current);
 						UpdateStateAndNotify(item, state => state with { Pending = pending });
-						return;
+
+						ToolCommand uninstallCommand = new ToolCommand { FileName = "msiexec.exe" };
+						uninstallCommand.Arguments = new List<string> { "/x", $"{{{item._state.MsiProductId}}}" };
+
+						int exitCode = await RunCommandAsync(item.Id.ToString(), uninstallCommand, current.Dir.FullName, cancellationToken);
+						if (exitCode == 1602)
+						{
+							// User cancelled (https://learn.microsoft.com/en-us/windows/win32/msi/error-codes)
+							UpdateStateAndNotify(item, state => state with { Pending = null });
+							return;
+						}
+						else if (exitCode != 0 || IsMsiInstalled(item._state.MsiProductId))
+						{
+							pending = new PendingToolDeploymentInfo(true, $"Uninstall failed ({exitCode}).", current);
+							UpdateStateAndNotify(item, state => state with { Pending = pending });
+							return;
+						}
+					}
+					finally
+					{
+						_msiSemaphore.Release();
 					}
 				}
 			}
