Unreal Toolbox: Fix tool catalog auto-update setting not being initialized to value from settings.

#rnx

diff --git a/Engine/Source/Programs/UnrealToolbox/ToolCatalogViewModel.cs b/Engine/Source/Programs/UnrealToolbox/ToolCatalogViewModel.cs
--- a/Engine/Source/Programs/UnrealToolbox/ToolCatalogViewModel.cs
+++ b/Engine/Source/Programs/UnrealToolbox/ToolCatalogViewModel.cs
@@ -31,6 +31,8 @@ namespace UnrealToolbox
 				_catalog = catalog;
 				_catalog.OnItemsChanged += RefreshItems;
 
+				AutoUpdate = catalog.AutoUpdate;
+
 				Items.CollectionChanged += (_, _) => OnPropertyChanged(nameof(ShowTools));
 
 				RefreshItems();
