Unreal Toolbox: Set the ManualInstall flag for any tool that has an MSI installer, to prevent it automatically triggering an update when new versions become available.

#rnx

diff --git a/Engine/Source/Programs/UnrealToolbox/ToolCatalog.cs b/Engine/Source/Programs/UnrealToolbox/ToolCatalog.cs
--- a/Engine/Source/Programs/UnrealToolbox/ToolCatalog.cs
+++ b/Engine/Source/Programs/UnrealToolbox/ToolCatalog.cs
@@ -514,6 +514,9 @@ namespace UnrealToolbox
 			}
 			else if (OperatingSystem.IsWindows() && !String.IsNullOrEmpty(item._state.MsiProductId))
 			{
+				// Make sure we don't automatically auto-update this tool since it has an interactive installer
+				toolConfig.ManualInstall = true;
+
 				// Wait to run the installer
 				if (!await _msiSemaphore.WaitAsync(0, cancellationToken))
 				{
