Unreal Toolbox: Fix issue where MSI installers would not run if there was no connection to the server (MSI product id wasn't being saved to config file for use offline). Also fixed an issue where we would assume that uninstall actions would succeed, and remove the toolbox entry even if the user cancelled or installer failed.

#rnx

diff --git a/Engine/Source/Programs/UnrealToolbox/ToolCatalog.cs b/Engine/Source/Programs/UnrealToolbox/ToolCatalog.cs
--- a/Engine/Source/Programs/UnrealToolbox/ToolCatalog.cs
+++ b/Engine/Source/Programs/UnrealToolbox/ToolCatalog.cs
@@ -166,6 +166,7 @@ namespace UnrealToolbox
 				Id = id;
 				Name = state.Name;
 				Description = state.Description;
+				MsiProductId = state.MsiProductId;
 				Current = (state.Current == null) ? null : new JsonDeploymentState(state.Current);
 			}
 		}
@@ -174,7 +175,6 @@ namespace UnrealToolbox
 		{
 			public ToolDeploymentId Id { get; set; }
 			public string Version { get; set; }
-			public string? MsiProductId { get; set; }
 
 			public JsonDeploymentState()
 			{
@@ -497,7 +497,10 @@ namespace UnrealToolbox
 		{
 			ToolConfig toolConfig = LoadToolConfig(toolDir);
 
-			await UninstallInternalAsync(item, cancellationToken);
+			if (!await UninstallInternalAsync(item, cancellationToken))
+			{
+				return;
+			}
 
 			if (toolConfig.InstallCommand != null)
 			{
@@ -569,13 +572,18 @@ namespace UnrealToolbox
 			return !String.IsNullOrEmpty(value);
 		}
 
-		async Task UninstallAsync(Item item, CancellationToken cancellationToken)
+		async Task<bool> UninstallAsync(Item item, CancellationToken cancellationToken)
 		{
-			await UninstallInternalAsync(item, cancellationToken);
+			if (!await UninstallInternalAsync(item, cancellationToken))
+			{
+				return false;
+			}
+
 			UpdateStateAndNotify(item, state => state with { Pending = null, Current = null });
+			return true;
 		}
 
-		async Task UninstallInternalAsync(Item item, CancellationToken cancellationToken)
+		async Task<bool> UninstallInternalAsync(Item item, CancellationToken cancellationToken)
 		{
 			CurrentToolDeploymentInfo? current = item._state.Current;
 			if (current != null)
@@ -609,18 +617,18 @@ namespace UnrealToolbox
 						ToolCommand uninstallCommand = new ToolCommand { FileName = "msiexec.exe" };
 						uninstallCommand.Arguments = new List<string> { "/x", $"{{{item._state.MsiProductId}}}" };
 
-						int exitCode = await RunCommandAsync(item.Id.ToString(), uninstallCommand, current.Dir.FullName, cancellationToken);
+						int exitCode = await RunCommandAsync(item.Id.ToString(), uninstallCommand, Directory.GetCurrentDirectory(), cancellationToken);
 						if (exitCode == 1602)
 						{
 							// User cancelled (https://learn.microsoft.com/en-us/windows/win32/msi/error-codes)
 							UpdateStateAndNotify(item, state => state with { Pending = null });
-							return;
+							return false;
 						}
 						else if (exitCode != 0 || IsMsiInstalled(item._state.MsiProductId))
 						{
 							pending = new PendingToolDeploymentInfo(true, $"Uninstall failed ({exitCode}).", current);
 							UpdateStateAndNotify(item, state => state with { Pending = pending });
-							return;
+							return false;
 						}
 					}
 					finally
@@ -629,6 +637,7 @@ namespace UnrealToolbox
 					}
 				}
 			}
+			return true;
 		}
 
 		async Task<int> RunCommandAsync(string toolName, ToolCommand command, string workingDir, CancellationToken cancellationToken)
