Unreal Toolbox: Check that the MSI product id is null before attempting to auto-update a tool. The modified ManualInstall flag is not persisted between runs.

#rnx

diff --git a/Engine/Source/Programs/UnrealToolbox/ToolCatalog.cs b/Engine/Source/Programs/UnrealToolbox/ToolCatalog.cs
--- a/Engine/Source/Programs/UnrealToolbox/ToolCatalog.cs
+++ b/Engine/Source/Programs/UnrealToolbox/ToolCatalog.cs
@@ -517,9 +517,6 @@ namespace UnrealToolbox
 			}
 			else if (OperatingSystem.IsWindows() && !String.IsNullOrEmpty(item._state.MsiProductId))
 			{
-				// Make sure we don't automatically auto-update this tool since it has an interactive installer
-				toolConfig.ManualInstall = true;
-
 				// Wait to run the installer
 				if (!await _msiSemaphore.WaitAsync(0, cancellationToken))
 				{
@@ -970,7 +967,7 @@ namespace UnrealToolbox
 				{
 					foreach (Item item in _items.Values)
 					{
-						if (item._state.Current != null && !item._state.Current.Config.ManualInstall)
+						if (item._state.Current != null && !item._state.Current.Config.ManualInstall && String.IsNullOrEmpty(item._state.MsiProductId))
 						{
 							if (item._state.Latest != null && item._state.Latest.Id != item._state.Current.Id)
 							{
