Unreal Toolbox: Re-enable IDE0005 warnings (unused using directives) and fix current violations.

#rnx

diff --git a/Engine/Source/Programs/UnrealToolbox/AboutPage.axaml.cs b/Engine/Source/Programs/UnrealToolbox/AboutPage.axaml.cs
--- a/Engine/Source/Programs/UnrealToolbox/AboutPage.axaml.cs
+++ b/Engine/Source/Programs/UnrealToolbox/AboutPage.axaml.cs
@@ -1,9 +1,6 @@
 // Copyright Epic Games, Inc. All Rights Reserved.
 
-using Avalonia;
 using Avalonia.Controls;
-using Avalonia.Markup.Xaml;
-using Microsoft.Extensions.DependencyInjection;
 
 namespace UnrealToolbox
 {
diff --git a/Engine/Source/Programs/UnrealToolbox/App.axaml.cs b/Engine/Source/Programs/UnrealToolbox/App.axaml.cs
--- a/Engine/Source/Programs/UnrealToolbox/App.axaml.cs
+++ b/Engine/Source/Programs/UnrealToolbox/App.axaml.cs
@@ -1,18 +1,16 @@
 // Copyright Epic Games, Inc. All Rights Reserved.
 
 using System.Diagnostics;
+using System.Runtime.InteropServices;
 using Avalonia;
 using Avalonia.Controls;
 using Avalonia.Controls.ApplicationLifetimes;
-using Avalonia.Interactivity;
 using Avalonia.Markup.Xaml;
 using Avalonia.Threading;
 using EpicGames.Core;
+using Microsoft.Extensions.DependencyInjection;
 using UnrealToolbox.Plugins.HordeAgent;
 using UnrealToolbox.Plugins.HordeProxy;
-using Microsoft.Extensions.DependencyInjection;
-using DesktopNotifications;
-using System.Runtime.InteropServices;
 
 namespace UnrealToolbox
 {
@@ -104,7 +102,7 @@ namespace UnrealToolbox
 			using EventWaitHandle closeEvent = new EventWaitHandle(false, EventResetMode.AutoReset, Program.CloseEventName);
 			using EventWaitHandle refreshEvent = new EventWaitHandle(false, EventResetMode.AutoReset, Program.RefreshEventName);
 			using EventWaitHandle settingsEvent = new EventWaitHandle(false, EventResetMode.AutoReset, Program.SettingsEventName);
-			for (; ;)
+			for (; ; )
 			{
 				int index = WaitHandle.WaitAny(new[] { closeEvent, refreshEvent, settingsEvent, _settingsThreadStop! });
 				if (index == 0)
@@ -317,7 +315,7 @@ namespace UnrealToolbox
 			TrayAppPluginState state = TrayAppPluginState.Undefined;
 			List<string> messages = new List<string>();
 			string? pluginName = null;
-			string? errorMessage = null; 
+			string? errorMessage = null;
 
 			foreach (ITrayAppPlugin plugin in _serviceProvider.GetServices<ITrayAppPlugin>())
 			{
diff --git a/Engine/Source/Programs/UnrealToolbox/JsonConfig.cs b/Engine/Source/Programs/UnrealToolbox/JsonConfig.cs
--- a/Engine/Source/Programs/UnrealToolbox/JsonConfig.cs
+++ b/Engine/Source/Programs/UnrealToolbox/JsonConfig.cs
@@ -1,12 +1,7 @@
 // Copyright Epic Games, Inc. All Rights Reserved.
 
-using System;
-using System.Collections.Generic;
-using System.Linq;
-using System.Text;
 using System.Text.Json;
 using System.Text.Json.Serialization;
-using System.Threading.Tasks;
 using EpicGames.Core;
 
 namespace UnrealToolbox
diff --git a/Engine/Source/Programs/UnrealToolbox/Notifications/Notifications.cs b/Engine/Source/Programs/UnrealToolbox/Notifications/Notifications.cs
--- a/Engine/Source/Programs/UnrealToolbox/Notifications/Notifications.cs
+++ b/Engine/Source/Programs/UnrealToolbox/Notifications/Notifications.cs
@@ -1,17 +1,9 @@
 // Copyright Epic Games, Inc. All Rights Reserved.
 
-using Avalonia;
-using Avalonia.Controls;
-using Avalonia.Controls.ApplicationLifetimes;
-using Avalonia.Platform;
-using DesktopNotifications;
-using DesktopNotifications.FreeDesktop;
-using DesktopNotifications.Windows;
-using Microsoft.Extensions.Logging;
-using System;
 using System.Reflection;
 using System.Runtime.InteropServices;
-using UnrealToolbox;
+using DesktopNotifications;
+using DesktopNotifications.Windows;
 
 namespace UnrealToolbox
 {
@@ -54,7 +46,7 @@ namespace UnrealToolbox
 			{
 				// WindowsApplicationContext.FromCurrentProcess() has side effects of creating start menu items, and changing the app user model id to the executing assembly, which can be dotnet.exe
 				// WindowsApplicationContext context = WindowsApplicationContext.FromCurrentProcess();
-				WindowsApplicationContext? context = Activator.CreateInstance(type: typeof(WindowsApplicationContext), bindingAttr: BindingFlags.Instance | BindingFlags.NonPublic, binder: null, args: new object[] { "Unreal Toolbox", "Unreal Toolbox" },culture: null) as WindowsApplicationContext;
+				WindowsApplicationContext? context = Activator.CreateInstance(type: typeof(WindowsApplicationContext), bindingAttr: BindingFlags.Instance | BindingFlags.NonPublic, binder: null, args: new object[] { "Unreal Toolbox", "Unreal Toolbox" }, culture: null) as WindowsApplicationContext;
 				SetCurrentProcessExplicitAppUserModelID("Unreal Toolbox");
 
 				_platformManager = new WindowsNotificationManager(context);
diff --git a/Engine/Source/Programs/UnrealToolbox/Plugins/HordeProxy/HordeProxyPlugin.cs b/Engine/Source/Programs/UnrealToolbox/Plugins/HordeProxy/HordeProxyPlugin.cs
--- a/Engine/Source/Programs/UnrealToolbox/Plugins/HordeProxy/HordeProxyPlugin.cs
+++ b/Engine/Source/Programs/UnrealToolbox/Plugins/HordeProxy/HordeProxyPlugin.cs
@@ -1,8 +1,6 @@
 // Copyright Epic Games, Inc. All Rights Reserved.
 
 using System.Net;
-using System.Text.Json;
-using System.Text.Json.Serialization;
 using Avalonia.Controls;
 using EpicGames.Core;
 using EpicGames.Horde;
diff --git a/Engine/Source/Programs/UnrealToolbox/Plugins/HordeProxy/HordeProxySettingsViewModel.cs b/Engine/Source/Programs/UnrealToolbox/Plugins/HordeProxy/HordeProxySettingsViewModel.cs
--- a/Engine/Source/Programs/UnrealToolbox/Plugins/HordeProxy/HordeProxySettingsViewModel.cs
+++ b/Engine/Source/Programs/UnrealToolbox/Plugins/HordeProxy/HordeProxySettingsViewModel.cs
@@ -1,9 +1,6 @@
 // Copyright Epic Games, Inc. All Rights Reserved.
 
 using System.ComponentModel;
-using System.Diagnostics;
-using Avalonia.Platform.Storage;
-using Avalonia.Threading;
 using CommunityToolkit.Mvvm.ComponentModel;
 
 namespace UnrealToolbox.Plugins.HordeProxy
diff --git a/Engine/Source/Programs/UnrealToolbox/Program.cs b/Engine/Source/Programs/UnrealToolbox/Program.cs
--- a/Engine/Source/Programs/UnrealToolbox/Program.cs
+++ b/Engine/Source/Programs/UnrealToolbox/Program.cs
@@ -1,8 +1,5 @@
 // Copyright Epic Games, Inc. All Rights Reserved.
 
-using System.Diagnostics;
-using System.Reflection;
-using System.Runtime.InteropServices;
 using Avalonia;
 using Avalonia.Controls;
 using EpicGames.Core;
diff --git a/Engine/Source/Programs/UnrealToolbox/ToolCatalog.cs b/Engine/Source/Programs/UnrealToolbox/ToolCatalog.cs
--- a/Engine/Source/Programs/UnrealToolbox/ToolCatalog.cs
+++ b/Engine/Source/Programs/UnrealToolbox/ToolCatalog.cs
@@ -1,15 +1,14 @@
 // Copyright Epic Games, Inc. All Rights Reserved.
 
+using System.Diagnostics;
+using System.Runtime.InteropServices;
+using System.Text.Json;
 using EpicGames.Core;
 using EpicGames.Horde;
 using EpicGames.Horde.Storage.Nodes;
 using EpicGames.Horde.Tools;
 using Microsoft.Extensions.Logging;
 using Microsoft.Win32;
-using System.Diagnostics;
-using System.Runtime.InteropServices;
-using System.Text;
-using System.Text.Json;
 
 #nullable enable
 
diff --git a/Engine/Source/Programs/UnrealToolbox/ToolCatalogViewModel.cs b/Engine/Source/Programs/UnrealToolbox/ToolCatalogViewModel.cs
--- a/Engine/Source/Programs/UnrealToolbox/ToolCatalogViewModel.cs
+++ b/Engine/Source/Programs/UnrealToolbox/ToolCatalogViewModel.cs
@@ -4,7 +4,6 @@ using System.Collections.ObjectModel;
 using System.ComponentModel;
 using System.Diagnostics;
 using CommunityToolkit.Mvvm.ComponentModel;
-using EpicGames.Core;
 
 namespace UnrealToolbox
 {
diff --git a/Engine/Source/Programs/UnrealToolbox/UnrealToolbox.csproj b/Engine/Source/Programs/UnrealToolbox/UnrealToolbox.csproj
--- a/Engine/Source/Programs/UnrealToolbox/UnrealToolbox.csproj
+++ b/Engine/Source/Programs/UnrealToolbox/UnrealToolbox.csproj
@@ -12,7 +12,7 @@
     <Configurations>Debug;Release;Analyze</Configurations>
     <AssemblyName>UnrealToolbox</AssemblyName>
     <AvaloniaUseCompiledBindingsByDefault>true</AvaloniaUseCompiledBindingsByDefault>
-    <NoWarn>IDE0005;CA1501</NoWarn>
+    <NoWarn>CA1501</NoWarn>
     <AllowUnsafeBlocks>true</AllowUnsafeBlocks>
     <ApplicationIcon>Resources/StatusNormal.ico</ApplicationIcon>
   </PropertyGroup>
