Unreal Toolbox: Provide more context for connection errors when trying to connect to Horde, and allow cancelling and refreshing connection attempts through the UI.

#rnx

diff --git a/Engine/Source/Programs/UnrealToolbox/GeneralSettingsPage.axaml b/Engine/Source/Programs/UnrealToolbox/GeneralSettingsPage.axaml
--- a/Engine/Source/Programs/UnrealToolbox/GeneralSettingsPage.axaml
+++ b/Engine/Source/Programs/UnrealToolbox/GeneralSettingsPage.axaml
@@ -21,13 +21,18 @@
 			<TextBlock Classes="SettingHeader">Horde Server</TextBlock>
 
 			<Border Classes="Card">
-				<Grid ColumnDefinitions="*,Auto,Auto" RowDefinitions="Auto,Auto">
+				<Grid ColumnDefinitions="*,Auto,Auto,Auto" RowDefinitions="Auto">
 					<StackPanel Grid.Column="0" Margin="0">
 						<HyperlinkButton Grid.Column="0" Grid.Row="0" Padding="0" FontSize="14" Content="{Binding ServerUrl}" Command="{Binding OpenServer}"/>
-						<TextBlock Grid.Column="0" Grid.Row="1" Padding="1" FontSize="12" Text="{Binding ServerStatus}"/>
+						<Grid ColumnDefinitions="Auto,Auto" RowDefinitions="Auto">
+							<ui:ProgressRing Grid.Column="0" Grid.Row="0" Width="16" Height="16" Margin="0 0 6 0" IsVisible="{Binding IsConnecting}" VerticalAlignment="Center"/>
+							<TextBlock Grid.Column="1" Grid.Row="0" Padding="1" FontSize="12" Text="{Binding ServerStatus}" VerticalAlignment="Top" TextTrimming="CharacterEllipsis"/>
+						</Grid>
 					</StackPanel>
-					<Button Grid.Column="1" Command="{Binding LoginAsync}" IsEnabled="{Binding IsLoginEnabled}" Margin="0 0 5 0">Connect</Button>
-					<Button Grid.Column="2" Command="{Binding ConfigureServerAsync}">Modify...</Button>
+
+					<Button Grid.Column="1" Command="{Binding StartRefresh}" IsEnabled="{Binding !IsConnecting}" Margin="0 0 5 0">Refresh</Button>
+					<Button Grid.Column="2" Command="{Binding CancelRefresh}" IsEnabled="{Binding IsConnecting}" Margin="0 0 5 0">Cancel</Button>
+					<Button Grid.Column="3" Command="{Binding ConfigureServerAsync}">Modify...</Button>
 				</Grid>
 			</Border>
 
diff --git a/Engine/Source/Programs/UnrealToolbox/GeneralSettingsViewModel.cs b/Engine/Source/Programs/UnrealToolbox/GeneralSettingsViewModel.cs
--- a/Engine/Source/Programs/UnrealToolbox/GeneralSettingsViewModel.cs
+++ b/Engine/Source/Programs/UnrealToolbox/GeneralSettingsViewModel.cs
@@ -2,7 +2,9 @@
 
 using System.Diagnostics;
 using Avalonia.Controls;
+using Avalonia.Threading;
 using CommunityToolkit.Mvvm.ComponentModel;
+using EpicGames.Core;
 using EpicGames.Horde;
 using FluentAvalonia.UI.Controls;
 using Microsoft.Extensions.DependencyInjection;
@@ -11,7 +13,10 @@ namespace UnrealToolbox
 {
 	partial class GeneralSettingsViewModel : ObservableObject, IDisposable
 	{
-		record StateRecord(Uri ServerUrl, bool Connected);
+		record class ConnectionState(string ServerUrl, string ServerStatus)
+		{
+			public static ConnectionState NoServerConfigured { get; } = new ConnectionState("No client configured", "Status unavailable");
+		}
 
 		readonly SettingsContext _context;
 		readonly IHordeClientProvider _hordeClientProvider;
@@ -23,7 +28,10 @@ namespace UnrealToolbox
 		string _serverStatus = String.Empty;
 
 		[ObservableProperty]
-		bool? _isLoginEnabled;
+		bool _isConnecting;
+
+		bool _requestRefresh;
+		BackgroundTask<ConnectionState>? _connectTask;
 
 		public ToolCatalogViewModel ToolCatalog { get; }
 
@@ -37,110 +45,175 @@ namespace UnrealToolbox
 			_context = context;
 
 			_hordeClientProvider = context.ServiceProvider.GetRequiredService<IHordeClientProvider>();
-			_hordeClientProvider.OnStateChanged += OnStateChangedAsync;
-			_hordeClientProvider.OnAccessTokenStateChanged += OnStateChangedAsync;
+			_hordeClientProvider.OnStateChanged += StartRefresh;
+//			_hordeClientProvider.OnAccessTokenStateChanged += StartRefresh;
 
 			ToolCatalog = new ToolCatalogViewModel(context.ServiceProvider.GetService<IToolCatalog>());
 
-			OnStateChangedAsync();
+			ServerUrl = HordeOptions.GetDefaultServerUrl()?.ToString() ?? "Unknown";
+			ServerStatus = "Unknown";
+
+			StartRefresh();
 		}
 
-		public void Dispose()
+		public async void Dispose()
 		{
-			_hordeClientProvider.OnStateChanged -= OnStateChangedAsync;
-			_hordeClientProvider.OnAccessTokenStateChanged -= OnStateChangedAsync;
+			_hordeClientProvider.OnStateChanged -= StartRefresh;
+			//			_hordeClientProvider.OnAccessTokenStateChanged -= StartRefresh;
+
+			if (_connectTask != null)
+			{
+				await _connectTask.DisposeAsync();
+				_connectTask = null;
+			}
 		}
 
-		async void OnStateChangedAsync()
+		void RunConnectionTask(Func<CancellationToken, Task<ConnectionState>> innerTask)
 		{
-			using IHordeClientRef? hordeClient = await _hordeClientProvider.GetClientRefAsync();
-			if (hordeClient == null)
+			BackgroundTask<ConnectionState>? prevConnectTask = _connectTask;
+
+			_connectTask = BackgroundTask.StartNew<ConnectionState>(async cancellationToken =>
 			{
-				ServerUrl = "No server configured";
-				ServerStatus = "Status unavailable";
-				IsLoginEnabled = false;
-			}
-			else
+				if (prevConnectTask != null)
+				{
+					try
+					{
+						await prevConnectTask.DisposeAsync();
+					}
+					catch (OperationCanceledException)
+					{
+						// Ignore
+					}
+				}
+
+				return await innerTask(cancellationToken);
+			});
+			_connectTask.Task?.ContinueWith(_ => Dispatcher.UIThread.Post(() => UpdateConnectionState()), TaskScheduler.Default);
+
+			ServerStatus = "Connecting...";
+			IsConnecting = true;
+		}
+
+		void UpdateConnectionState()
+		{
+			if (_connectTask != null && (_connectTask.Task?.IsCompleted ?? true))
 			{
-				ServerUrl = hordeClient.Client.ServerUrl.ToString();
-				ServerStatus = hordeClient.Client.HasValidAccessToken() ? "Authenticated" : "Session Expired";
-				IsLoginEnabled = !hordeClient.Client.HasValidAccessToken();
+				ConnectionState? connectionState = null;
+				if (_connectTask.Task != null)
+				{
+					try
+					{
+						_connectTask.Task.TryGetResult(out connectionState);
+					}
+					catch (Exception ex)
+					{
+						connectionState = new ConnectionState(ServerUrl, $"Unable to connect: {ex.Message}");
+					}
+				}
+
+				if(connectionState != null)
+				{
+					ServerUrl = connectionState.ServerUrl;
+					ServerStatus = connectionState.ServerStatus;
+				}
+
+				IsConnecting = false;
+
+				_ = _connectTask.DisposeAsync().AsTask();
+				_connectTask = null;
 			}
 		}
 
 		public void OpenServer()
+			=> OpenBrowser(ServerUrl);
+
+		public void StartRefresh()
 		{
-			OpenBrowser(ServerUrl);
+			if (!_requestRefresh)
+			{
+				_requestRefresh = true;
+				Dispatcher.UIThread.Post(() => UpdateRefresh());
+			}
 		}
 
-		public async Task LoginAsync()
+		public void CancelRefresh()
+			=> RunConnectionTask(_ => Task.FromResult<ConnectionState>(new ConnectionState(ServerUrl, "Cancelled")));
+
+		void UpdateRefresh()
 		{
-			using IHordeClientRef? hordeClientRef = await _hordeClientProvider.GetClientRefAsync();
-			if (hordeClientRef != null)
+			if (_requestRefresh)
 			{
-				await hordeClientRef.Client.LoginAsync(true, CancellationToken.None);
+				_requestRefresh = false;
+				RunConnectionTask(HandleRefreshAsync);
 			}
 		}
 
-		public async Task ConfigureServerAsync()
+		async Task<ConnectionState> HandleRefreshAsync(CancellationToken cancellationToken)
 		{
-			string serverUrl = ServerUrl;
+			using IHordeClientRef hordeClientRef = _hordeClientProvider.GetClientRef();
 
-			string? errorMessage = null;
-			for (; ; )
+			IHordeClient client = hordeClientRef.Client;
+			try
 			{
-				TextBox serverUrlTextBox = new TextBox();
-				serverUrlTextBox.Text = serverUrl;
-				serverUrlTextBox.SelectionStart = serverUrl.Length;
-				serverUrlTextBox.SelectionEnd = serverUrl.Length;
-
-				Control content = serverUrlTextBox;
-				if (!String.IsNullOrEmpty(errorMessage))
+				bool result = await client.LoginAsync(true, cancellationToken);
+				if (!result)
 				{
-					TextBlock errorBlock = new TextBlock();
-					errorBlock.Text = errorMessage;
-
-					StackPanel stackPanel = new StackPanel() { Spacing = 10 };
-					stackPanel.Children.Add(serverUrlTextBox);
-					stackPanel.Children.Add(errorBlock);
-
-					content = stackPanel;
+					return new ConnectionState(client.ServerUrl.ToString(), "Login failed");
 				}
-
-				ContentDialog dialog = new ContentDialog()
+				else if (!client.HasValidAccessToken())
 				{
-					Title = "Connect to Server",
-					Content = content,
-					PrimaryButtonText = "Connect",
-					CloseButtonText = "Cancel"
-				};
-
-				ContentDialogResult result = await dialog.ShowAsync(_context.SettingsWindow);
-				if (result == ContentDialogResult.None)
+					return new ConnectionState(client.ServerUrl.ToString(), "Session expired");
+				}
+				else
 				{
-					return;
+					return new ConnectionState(client.ServerUrl.ToString(), "Authenticated");
 				}
+			}
+			catch (Exception ex)
+			{
+				return new ConnectionState(hordeClientRef.Client.ServerUrl.ToString(), $"Connection failed: {ex.Message}");
+			}
+		}
+		
+		public async Task ConfigureServerAsync()
+		{
+			string serverUrl = ServerUrl;
 
-				serverUrl = serverUrlTextBox.Text;
+			TextBox serverUrlTextBox = new TextBox();
+			serverUrlTextBox.Text = serverUrl;
+			serverUrlTextBox.SelectionStart = serverUrl.Length;
+			serverUrlTextBox.SelectionEnd = serverUrl.Length;
 
-				try
-				{
-					HordeOptions.SetDefaultServerUrl(new Uri(serverUrl));
+			ContentDialog dialog = new ContentDialog()
+			{
+				Title = "Connect to Server",
+				Content = serverUrlTextBox,
+				PrimaryButtonText = "Connect",
+				CloseButtonText = "Cancel"
+			};
+
+			ContentDialogResult result = await dialog.ShowAsync(_context.SettingsWindow);
+			if (result == ContentDialogResult.None)
+			{
+				return;
+			}
 
-					await _hordeClientProvider.RecreateAsync();
+			serverUrl = serverUrlTextBox.Text;
 
-					using IHordeClientRef? clientRef = await _hordeClientProvider.GetClientRefAsync();
-					if (clientRef != null)
-					{
-						await clientRef.Client.LoginAsync(true, CancellationToken.None);
-					}
+			try
+			{
+				Uri normalizedServerUrl = new Uri(serverUrl.Trim());
+				HordeOptions.SetDefaultServerUrl(normalizedServerUrl);
 
-					return;
-				}
-				catch (Exception ex)
-				{
-					errorMessage = ex.ToString();
-				}
+				ServerUrl = normalizedServerUrl.ToString();
+				ServerStatus = "Connecting...";
+
+				_hordeClientProvider.Reset(); // Will trigger a call to StartRefresh()
+			}
+			catch (Exception ex)
+			{
+				ServerStatus = $"Error: {ex.Message}";
+				RunConnectionTask(_ => Task.FromResult<ConnectionState>(new ConnectionState(ServerUrl, ServerStatus)));
 			}
 		}
 
diff --git a/Engine/Source/Programs/UnrealToolbox/HordeClientProvider.cs b/Engine/Source/Programs/UnrealToolbox/HordeClientProvider.cs
--- a/Engine/Source/Programs/UnrealToolbox/HordeClientProvider.cs
+++ b/Engine/Source/Programs/UnrealToolbox/HordeClientProvider.cs
@@ -1,6 +1,7 @@
 // Copyright Epic Games, Inc. All Rights Reserved.
 
 using System.Diagnostics;
+using EpicGames.Core;
 using EpicGames.Horde;
 using Microsoft.Extensions.DependencyInjection;
 using Microsoft.Extensions.Logging;
@@ -16,25 +17,35 @@ namespace UnrealToolbox
 		{
 			readonly TaskCompletionSource _refZeroTcs = new TaskCompletionSource(TaskCreationOptions.RunContinuationsAsynchronously);
 
-			readonly Task<IHordeClient?> _clientTask;
 			int _refCount;
-			bool _disposed;
+			ServiceProvider _serviceProvider;
+			IHordeClient _hordeClient;
 
-			public Task<IHordeClient?> ClientTask => _clientTask;
+			public IHordeClient Client => _hordeClient;
 
-			public HordeClientLifetime(Task<IHordeClient?> clientTask)
+			public HordeClientLifetime(ILoggerFactory loggerFactory)
 			{
-				_clientTask = clientTask;
 				_refCount = 1;
+
+				ServiceCollection serviceCollection = new ServiceCollection();
+				serviceCollection.AddHorde(options => options.AllowAuthPrompt = false);
+				serviceCollection.AddSingleton<ILoggerFactory>(loggerFactory);
+				serviceCollection.AddSingleton(typeof(ILogger<>), typeof(Logger<>));
+				_serviceProvider = serviceCollection.BuildServiceProvider();
+
+				_hordeClient = _serviceProvider.GetRequiredService<IHordeClient>();
 			}
 
 			public async ValueTask DisposeAsync()
 			{
-				if (!_disposed)
+				if (_serviceProvider != null)
 				{
 					Release();
 					await _refZeroTcs.Task;
-					_disposed = true;
+
+					await _serviceProvider.DisposeAsync();
+					_serviceProvider = null!;
+					_hordeClient = null!;
 				}
 			}
 
@@ -55,32 +66,29 @@ namespace UnrealToolbox
 
 		class HordeClientRef : IHordeClientRef
 		{
-			IHordeClient _client;
 			HordeClientLifetime? _lifetime;
 
 			public IHordeClient Client
-				=> _client ?? throw new ObjectDisposedException(null);
+				=> _lifetime?.Client ?? throw new ObjectDisposedException(null);
 
-			public HordeClientRef(IHordeClient client, HordeClientLifetime lifetime)
+			public HordeClientRef(HordeClientLifetime lifetime)
 			{
-				_client = client;
 				_lifetime = lifetime;
+				_lifetime.AddRef();
 			}
 
 			public void Dispose()
 			{
 				_lifetime?.Release();
 				_lifetime = null;
-				_client = null!;
 			}
 		}
 
 		readonly object _lockObject = new object();
 		readonly ILoggerFactory _loggerFactory;
-		ServiceProvider? _serviceProvider;
-		IHordeClient? _hordeClient;
-		HordeClientLifetime _lifetime;
-		readonly ILogger _logger;
+		readonly List<Task> _disposeTasks = new List<Task>();
+
+		HordeClientLifetime? _lifetime;
 
 		/// <summary>
 		/// Event signalled whenever the connection state changes
@@ -95,67 +103,41 @@ namespace UnrealToolbox
 		/// <summary>
 		/// Constructor
 		/// </summary>
-		public HordeClientProvider(ILoggerFactory loggerFactory, ILogger<HordeClientProvider> logger)
+		public HordeClientProvider(ILoggerFactory loggerFactory)
 		{
 			_loggerFactory = loggerFactory;
-			_logger = logger;
-
-			CreateServices();
-
-			_lifetime = new HordeClientLifetime(Task.FromResult(_hordeClient));
+			_lifetime = new HordeClientLifetime(_loggerFactory);
 		}
 
 		/// <inheritdoc/>
 		public async ValueTask DisposeAsync()
 		{
-			await _lifetime.DisposeAsync();
-			await DestroyServicesAsync();
+			if (_lifetime != null)
+			{
+				DestroyLifetime();
+			}
+			await Task.WhenAll(_disposeTasks);
 		}
 
 		/// <inheritdoc/>
-		public async Task<IHordeClientRef?> GetClientRefAsync()
+		public IHordeClientRef GetClientRef()
 		{
-			HordeClientLifetime? lifetime = null;
-			try
-			{
-				lock (_lockObject)
-				{
-					lifetime = _lifetime;
-					lifetime.AddRef();
-				}
-
-				IHordeClient? client = await lifetime.ClientTask;
-				if (client == null)
-				{
-					lifetime.Release();
-					return null;
-				}
-
-				return new HordeClientRef(client, lifetime);
-			}
-			catch
+			lock (_lockObject)
 			{
-				lifetime?.Release();
-				throw;
+				return new HordeClientRef(_lifetime ?? throw new ObjectDisposedException("HordeClientProvider"));
 			}
 		}
 
 		/// <inheritdoc/>
-		public async Task RecreateAsync()
+		public void Reset()
 		{
 			lock (_lockObject)
 			{
-				HordeClientLifetime prevLifetime = _lifetime;
-				Task<IHordeClient?> clientTask = Task.Run(async () =>
-				{
-					await prevLifetime.DisposeAsync();
-					await DestroyServicesAsync();
-					CreateServices();
-					return _hordeClient;
-				});
-				_lifetime = new HordeClientLifetime(clientTask);
+				DestroyLifetime();
+				CreateLifetime();
 			}
-			await _lifetime.ClientTask;
+
+			OnStateChanged?.Invoke();
 		}
 
 		void OnAccessTokenStateChangedInternal()
@@ -163,42 +145,24 @@ namespace UnrealToolbox
 			OnAccessTokenStateChanged?.Invoke();
 		}
 
-		void CreateServices()
+		void CreateLifetime()
 		{
-			Debug.Assert(_serviceProvider == null);
-			Debug.Assert(_hordeClient == null);
-
-			ServiceCollection serviceCollection = new ServiceCollection();
-			serviceCollection.AddHorde(options => options.AllowAuthPrompt = false);
-			serviceCollection.AddSingleton<ILoggerFactory>(_loggerFactory);
-			serviceCollection.AddSingleton(typeof(ILogger<>), typeof(Logger<>));
-			_serviceProvider = serviceCollection.BuildServiceProvider();
-
-			try
-			{
-				_hordeClient = _serviceProvider.GetRequiredService<IHordeClient>();
-				_hordeClient.OnAccessTokenStateChanged += OnAccessTokenStateChangedInternal;
-			}
-			catch (Exception ex)
-			{
-				_logger.LogWarning(ex, "Horde client could not be created: {Message}", ex.Message);
-			}
+			Debug.Assert(_lifetime == null);
 
-			OnStateChanged?.Invoke();
+			_lifetime = new HordeClientLifetime(_loggerFactory);
+			_lifetime.Client.OnAccessTokenStateChanged += OnAccessTokenStateChangedInternal;
 		}
 
-		async ValueTask DestroyServicesAsync()
+		void DestroyLifetime()
 		{
-			if (_hordeClient != null)
+			if (_lifetime != null)
 			{
-				_hordeClient.OnAccessTokenStateChanged -= OnAccessTokenStateChangedInternal;
-				_hordeClient = null;
-			}
-			if (_serviceProvider != null)
-			{
-				await _serviceProvider.DisposeAsync();
-				_serviceProvider = null;
+				_lifetime.Client.OnAccessTokenStateChanged -= OnAccessTokenStateChangedInternal;
+				_disposeTasks.Add(_lifetime.DisposeAsync().AsTask());
+				_lifetime = null;
 			}
+
+			AsyncUtils.RemoveCompleteTasks(_disposeTasks);
 		}
 	}
 }
diff --git a/Engine/Source/Programs/UnrealToolbox/IHordeClientProvider.cs b/Engine/Source/Programs/UnrealToolbox/IHordeClientProvider.cs
--- a/Engine/Source/Programs/UnrealToolbox/IHordeClientProvider.cs
+++ b/Engine/Source/Programs/UnrealToolbox/IHordeClientProvider.cs
@@ -31,13 +31,13 @@ namespace UnrealToolbox
 		event Action? OnAccessTokenStateChanged;
 
 		/// <summary>
-		/// Gets a reference to the current client instance. These references should be kept as short as possible.
+		/// Resets the current client and creates a new one with the latest settings.
 		/// </summary>
-		Task<IHordeClientRef?> GetClientRefAsync();
+		void Reset();
 
 		/// <summary>
-		/// Recreates the client instance.
+		/// Gets a reference to the current client instance. These references should be kept as short as possible.
 		/// </summary>
-		Task RecreateAsync();
+		IHordeClientRef GetClientRef();
 	}
 }
diff --git a/Engine/Source/Programs/UnrealToolbox/SelfUpdateService.cs b/Engine/Source/Programs/UnrealToolbox/SelfUpdateService.cs
--- a/Engine/Source/Programs/UnrealToolbox/SelfUpdateService.cs
+++ b/Engine/Source/Programs/UnrealToolbox/SelfUpdateService.cs
@@ -286,11 +286,7 @@ namespace UnrealToolbox
 				return;
 			}
 
-			using IHordeClientRef? clientRef = await _hordeClientProvider.GetClientRefAsync().WaitAsync(cancellationToken);
-			if (clientRef == null)
-			{
-				return;
-			}
+			using IHordeClientRef clientRef = _hordeClientProvider.GetClientRef();
 
 			ITool? tool = await clientRef.Client.Tools.GetAsync(ToolId, cancellationToken);
 			if (tool == null || tool.Deployments.Count == 0)
diff --git a/Engine/Source/Programs/UnrealToolbox/SettingsWindow.axaml.cs b/Engine/Source/Programs/UnrealToolbox/SettingsWindow.axaml.cs
--- a/Engine/Source/Programs/UnrealToolbox/SettingsWindow.axaml.cs
+++ b/Engine/Source/Programs/UnrealToolbox/SettingsWindow.axaml.cs
@@ -16,6 +16,10 @@ namespace UnrealToolbox
 	{
 		readonly IServiceProvider? _serviceProvider;
 		readonly Dictionary<string, ITrayAppPlugin> _typeToPlugin = new Dictionary<string, ITrayAppPlugin>();
+		readonly Dictionary<ITrayAppPlugin, Control> _pluginPages = new Dictionary<ITrayAppPlugin, Control>();
+
+		AboutPage? _aboutPage;
+		GeneralSettingsPage? _generalSettingsPage;
 
 		public SettingsWindow()
 			: this(null!)
@@ -94,17 +98,31 @@ namespace UnrealToolbox
 				string? typeName = (string?)nvi.Tag;
 				if (typeName != null && _typeToPlugin.TryGetValue(typeName, out ITrayAppPlugin? plugin))
 				{
-					_navView.Content = plugin.CreateSettingsPage(context);
+					Control pluginPage = CreatePluginPage(plugin, context);
+					_navView.Content = pluginPage;
 				}
 				else if (nvi == _navView.FooterMenuItems.FirstOrDefault())
 				{
-					_navView.Content = new AboutPage();
+					_aboutPage ??= new AboutPage();
+					_navView.Content = _aboutPage;
 				}
 				else
 				{
-					_navView.Content = new GeneralSettingsPage(context);
+					_generalSettingsPage ??= new GeneralSettingsPage(context);
+					_navView.Content = _generalSettingsPage;
 				}
 			}
 		}
+
+		private Control CreatePluginPage(ITrayAppPlugin plugin, SettingsContext context)
+		{
+			Control? pluginPage;
+			if (!_pluginPages.TryGetValue(plugin, out pluginPage))
+			{
+				pluginPage = plugin.CreateSettingsPage(context);
+				_pluginPages.Add(plugin, pluginPage);
+			}
+			return pluginPage;
+		}
 	}
 }
diff --git a/Engine/Source/Programs/UnrealToolbox/ToolCatalog.cs b/Engine/Source/Programs/UnrealToolbox/ToolCatalog.cs
--- a/Engine/Source/Programs/UnrealToolbox/ToolCatalog.cs
+++ b/Engine/Source/Programs/UnrealToolbox/ToolCatalog.cs
@@ -884,40 +884,55 @@ namespace UnrealToolbox
 			while (!cancellationToken.IsCancellationRequested)
 			{
 				Task updateTask = _updateEvent.Task;
-				try
-				{
-					await PollForUpdatesOnceAsync(cancellationToken);
-				}
-				catch (OperationCanceledException)
-				{
-					break;
-				}
-				catch (Exception ex)
+
+				using (CancellationTokenSource cancellationSource = CancellationTokenSource.CreateLinkedTokenSource(cancellationToken))
 				{
-					_logger.LogError(ex, "Exception while checking for tool updates: {Message}", ex.Message);
+					bool complete = false;
+					_ = updateTask.ContinueWith(_ =>
+					{
+						lock (cancellationSource)
+						{
+							if (!complete)
+							{
+								cancellationSource.Cancel();
+							}
+						}
+					}, TaskScheduler.Default);
+
+					try
+					{
+						await PollForUpdatesOnceAsync(cancellationSource.Token);
+					}
+					catch (OperationCanceledException)
+					{
+					}
+					catch (Exception ex)
+					{
+						_logger.LogError(ex, "Exception while checking for tool updates: {Message}", ex.Message);
+					}
+
+					lock (cancellationSource)
+					{
+						complete = true;
+					}
 				}
+
 				await Task.WhenAny(updateTask, Task.Delay(TimeSpan.FromMinutes(5.0), cancellationToken));
 			}
 		}
 
 		async Task PollForUpdatesOnceAsync(CancellationToken cancellationToken)
 		{
-			using IHordeClientRef? hordeClientRef = await _hordeClientProvider.GetClientRefAsync();
-			if (hordeClientRef == null)
+			using IHordeClientRef hordeClientRef = _hordeClientProvider.GetClientRef();
+
+			if (_serverUri != hordeClientRef.Client.ServerUrl)
 			{
 				SetTools(Enumerable.Empty<ITool>());
+				_serverUri = hordeClientRef.Client.ServerUrl;
 			}
-			else
-			{
-				if (_serverUri != hordeClientRef.Client.ServerUrl)
-				{
-					SetTools(Enumerable.Empty<ITool>());
-					_serverUri = hordeClientRef.Client.ServerUrl;
-				}
 
-				IEnumerable<ITool> tools = await hordeClientRef.Client.Tools.GetAllAsync(cancellationToken);
-				SetTools(tools);
-			}
+			IEnumerable<ITool> tools = await hordeClientRef.Client.Tools.GetAllAsync(cancellationToken);
+			SetTools(tools);
 		}
 
 		void SetTools(IEnumerable<ITool> tools)
