Unreal Toolbox: Show the agent plugin settings if the agent is connected, even if it's not installed by the MSI.

#rnx

diff --git a/Engine/Source/Programs/UnrealToolbox/Plugins/HordeAgent/HordeAgentPlugin.cs b/Engine/Source/Programs/UnrealToolbox/Plugins/HordeAgent/HordeAgentPlugin.cs
--- a/Engine/Source/Programs/UnrealToolbox/Plugins/HordeAgent/HordeAgentPlugin.cs
+++ b/Engine/Source/Programs/UnrealToolbox/Plugins/HordeAgent/HordeAgentPlugin.cs
@@ -29,6 +29,7 @@ namespace UnrealToolbox.Plugins.HordeAgent
 
 		byte[] _settingsData = Array.Empty<byte>();
 		HordeAgentSettings _settings = new HordeAgentSettings();
+		bool _pipeConnected;
 
 		public HordeAgentSettings Settings => _settings;
 
@@ -93,7 +94,7 @@ namespace UnrealToolbox.Plugins.HordeAgent
 		{
 			if (OperatingSystem.IsWindows())
 			{
-				bool enabled = ((Registry.GetValue("HKEY_LOCAL_MACHINE\\SOFTWARE\\Epic Games\\Horde\\Agent", "Installed", null) as int?) ?? 0) != 0;
+				bool enabled = _pipeConnected || ((Registry.GetValue("HKEY_LOCAL_MACHINE\\SOFTWARE\\Epic Games\\Horde\\Agent", "Installed", null) as int?) ?? 0) != 0;
 				if (enabled != IsEnabled)
 				{
 					IsEnabled = enabled;
@@ -267,7 +268,14 @@ namespace UnrealToolbox.Plugins.HordeAgent
 			{
 				try
 				{
-					await PollForStatusUpdatesAsync(cancellationToken);
+					try
+					{
+						await PollForStatusUpdatesAsync(cancellationToken);
+					}
+					finally
+					{
+						_pipeConnected = false;
+					}
 				}
 				catch (OperationCanceledException) when (cancellationToken.IsCancellationRequested)
 				{
@@ -493,6 +501,8 @@ namespace UnrealToolbox.Plugins.HordeAgent
 				SetStatus(new AgentStatusMessage(true, 0, "Connecting to agent..."));
 				await pipeClient.ConnectAsync(cancellationToken);
 
+				_pipeConnected = true;
+
 				SetStatus(new AgentStatusMessage(true, 0, "Waiting for status update."));
 				for (; ; )
 				{
diff --git a/Engine/Source/Programs/UnrealToolbox/SettingsWindow.axaml.cs b/Engine/Source/Programs/UnrealToolbox/SettingsWindow.axaml.cs
--- a/Engine/Source/Programs/UnrealToolbox/SettingsWindow.axaml.cs
+++ b/Engine/Source/Programs/UnrealToolbox/SettingsWindow.axaml.cs
@@ -32,6 +32,8 @@ namespace UnrealToolbox
 			_serviceProvider = serviceProvider;
 			_navView.SelectionChanged += NavView_SelectionChanged;
 
+			Activated += OnActivated;
+
 			Refresh();
 		}
 
@@ -68,15 +70,18 @@ namespace UnrealToolbox
 			}
 		}
 
-		protected override void OnGotFocus(GotFocusEventArgs e)
+		void OnActivated(object? sender, EventArgs e)
 		{
-			base.OnGotFocus(e);
-
 			if (_serviceProvider != null)
 			{
+				bool updated = false;
 				foreach (ITrayAppPlugin plugin in _serviceProvider.GetServices<ITrayAppPlugin>())
 				{
-					plugin.Refresh();
+					updated |= plugin.Refresh();
+				}
+				if (updated)
+				{
+					Refresh();
 				}
 			}
 		}
