Unreal Toolbox: Additional logging for upgrading to latest version.

#rnx

diff --git a/Engine/Source/Programs/UnrealToolbox/App.axaml.cs b/Engine/Source/Programs/UnrealToolbox/App.axaml.cs
--- a/Engine/Source/Programs/UnrealToolbox/App.axaml.cs
+++ b/Engine/Source/Programs/UnrealToolbox/App.axaml.cs
@@ -9,6 +9,7 @@ using Avalonia.Markup.Xaml;
 using Avalonia.Threading;
 using EpicGames.Core;
 using Microsoft.Extensions.DependencyInjection;
+using Microsoft.Extensions.Logging;
 using UnrealToolbox.Plugins.HordeAgent;
 using UnrealToolbox.Plugins.HordeProxy;
 
@@ -20,6 +21,7 @@ namespace UnrealToolbox
 	public sealed partial class App : Application, ITrayAppHost, IAsyncDisposable
 	{
 		readonly ServiceProvider _serviceProvider;
+		readonly ILogger _logger;
 
 		WindowIcon? _normalIcon;
 		WindowIcon? _busyIcon;
@@ -53,6 +55,7 @@ namespace UnrealToolbox
 			serviceCollection.AddSingleton<ITrayAppPlugin, HordeProxyPlugin>();
 
 			_serviceProvider = serviceCollection.BuildServiceProvider();
+			_logger = _serviceProvider.GetRequiredService<ILogger<App>>();
 		}
 
 		/// <inheritdoc/>
@@ -138,10 +141,12 @@ namespace UnrealToolbox
 		{
 			if (_settingsWindow == null)
 			{
+				_logger.LogInformation("Shutting down to install update");
 				((IClassicDesktopStyleApplicationLifetime)ApplicationLifetime!).Shutdown();
 			}
 			else
 			{
+				_logger.LogInformation("Scheduling update when settings window closes");
 				_shutdownOnClose = true;
 			}
 		}
@@ -295,6 +300,7 @@ namespace UnrealToolbox
 
 			if (_shutdownOnClose)
 			{
+				_logger.LogInformation("Running update due to settings window closing");
 				((IClassicDesktopStyleApplicationLifetime)ApplicationLifetime!).Shutdown();
 			}
 		}
diff --git a/Engine/Source/Programs/UnrealToolbox/SelfUpdateService.cs b/Engine/Source/Programs/UnrealToolbox/SelfUpdateService.cs
--- a/Engine/Source/Programs/UnrealToolbox/SelfUpdateService.cs
+++ b/Engine/Source/Programs/UnrealToolbox/SelfUpdateService.cs
@@ -310,10 +310,12 @@ namespace UnrealToolbox
 				await deployment.Content.ExtractAsync(updateDir.ToDirectoryInfo(), _logger, cancellationToken);
 
 				selfUpdate.UpdateVersion = updateVersion;
+				_logger.LogInformation("Extracted {UpdateVersion} to {UpdateDir}", updateVersion, updateDir);
 			}
 
-			if (!String.IsNullOrEmpty(selfUpdate.UpdateVersion))
+			if (selfUpdate.IsUpdatePending())
 			{
+				_logger.LogInformation("Triggering update to {UpdateVersion}", selfUpdate.UpdateVersion);
 				OnUpdateReady?.Invoke();
 			}
 		}
