Unreal Toolbox: Fix auto-update failing due to the initiating instance having two references to the single instance mutex.

#rnx

diff --git a/Engine/Source/Programs/UnrealToolbox/Program.cs b/Engine/Source/Programs/UnrealToolbox/Program.cs
--- a/Engine/Source/Programs/UnrealToolbox/Program.cs
+++ b/Engine/Source/Programs/UnrealToolbox/Program.cs
@@ -115,7 +115,7 @@ namespace UnrealToolbox
 
 		public SingleInstanceMutex(string name)
 		{
-			_mutex = new Mutex(true, name);
+			_mutex = new Mutex(false, name);
 		}
 
 		public void Release()
