Unreal Toolbox: Prevent updates triggering shutdown procedure if the new version will not launch due to being an older version.

#rnx

diff --git a/Engine/Source/Programs/UnrealToolbox/SelfUpdateService.cs b/Engine/Source/Programs/UnrealToolbox/SelfUpdateService.cs
--- a/Engine/Source/Programs/UnrealToolbox/SelfUpdateService.cs
+++ b/Engine/Source/Programs/UnrealToolbox/SelfUpdateService.cs
@@ -149,6 +149,12 @@ namespace UnrealToolbox
 		}
 
 		bool ShouldLaunchLatest(DirectoryReference latestDir)
+			=> ShouldUpdateTo(ReadVersion(latestDir));
+
+		/// <summary>
+		/// Test whether the current instance should update to the given version
+		/// </summary>
+		public bool ShouldUpdateTo(string? version)
 		{
 			// If this build is not versioned, do not upgrade
 			VersionNumber? currentVersionNumber;
@@ -159,7 +165,7 @@ namespace UnrealToolbox
 
 			// If the other build is not versioned, do not upgrade
 			VersionNumber? latestVersionNumber;
-			if (!TryParseVersion(ReadVersion(latestDir), out latestVersionNumber))
+			if (!TryParseVersion(version, out latestVersionNumber))
 			{
 				return false;
 			}
@@ -295,10 +301,12 @@ namespace UnrealToolbox
 			}
 
 			IToolDeployment deployment = tool.Deployments[^1];
-			_logger.LogInformation("Latest app version is {Id}", deployment.Id);
+			_logger.LogInformation("Latest deployment is {Id} ({Version})", deployment.Id, deployment.Version);
 
 			string updateVersion = deployment.Id.ToString();
-			if (!String.Equals(updateVersion, selfUpdate.LatestVersion, StringComparison.OrdinalIgnoreCase) && !String.Equals(updateVersion, selfUpdate.UpdateVersion, StringComparison.OrdinalIgnoreCase))
+			if (!String.Equals(updateVersion, selfUpdate.LatestVersion, StringComparison.OrdinalIgnoreCase) 
+				&& !String.Equals(updateVersion, selfUpdate.UpdateVersion, StringComparison.OrdinalIgnoreCase)
+				&& selfUpdate.ShouldUpdateTo(deployment.Version))
 			{
 				selfUpdate.UpdateVersion = null;
 
