Horde: Upgrade NuGet packages to fix newly released security vuln warning.

#rnx

diff --git a/Engine/Source/Programs/UnrealToolbox/SettingsWindow.axaml.cs b/Engine/Source/Programs/UnrealToolbox/SettingsWindow.axaml.cs
--- a/Engine/Source/Programs/UnrealToolbox/SettingsWindow.axaml.cs
+++ b/Engine/Source/Programs/UnrealToolbox/SettingsWindow.axaml.cs
@@ -1,7 +1,6 @@
 // Copyright Epic Games, Inc. All Rights Reserved.
 
 using Avalonia.Controls;
-using Avalonia.Input;
 using FluentAvalonia.UI.Controls;
 using Microsoft.Extensions.DependencyInjection;
 
@@ -61,7 +60,7 @@ namespace UnrealToolbox
 				_navView.FooterMenuItems.Clear();
 				_navView.FooterMenuItems.Add(new NavigationViewItem() { Content = "About", IconSource = new SymbolIconSource() { Symbol = Symbol.Help } });
 
-				_navView.SelectedItem = 
+				_navView.SelectedItem =
 					_navView.MenuItems.FirstOrDefault(x => Object.Equals((x as NavigationViewItem)?.Content, selectedItemContent))
 					?? _navView.FooterMenuItems.FirstOrDefault(x => Object.Equals((x as NavigationViewItem)?.Content, selectedItemContent))
 					?? _navView.MenuItems[0];
