Unreal Toolbox: Tweak notifications times to 2 minutes on different notifications, and 10 minutes for same notification

diff --git a/Engine/Source/Programs/UnrealToolbox/Notifications/Notifications.cs b/Engine/Source/Programs/UnrealToolbox/Notifications/Notifications.cs
--- a/Engine/Source/Programs/UnrealToolbox/Notifications/Notifications.cs
+++ b/Engine/Source/Programs/UnrealToolbox/Notifications/Notifications.cs
@@ -74,19 +74,24 @@ namespace UnrealToolbox
 		/// <param name="force"></param>
 		public void ShowNotification(string title, string body, bool force = false)
 		{
+			if (_platformManager == null)
+			{
+				return;
+			}
+
 			// spawn 
 			if (!force && _lastNotificationTime != null && !String.IsNullOrEmpty(_lastBody) && !String.IsNullOrEmpty(_lastTitle))
 			{
 				TimeSpan deltaTime = DateTime.Now - _lastNotificationTime.Value;
 
-				// don't show a new notification if already displayed one in last minute
-				if (deltaTime.TotalSeconds < 60)
+				// don't show a new notification if already displayed one in last 2 minutes
+				if (deltaTime.TotalSeconds < 120)
 				{
 					return;
 				}
 
-				// if the title and body are the same, wait 2 minutes
-				if ((deltaTime.TotalSeconds < 120) && title == _lastTitle && body == _lastBody)
+				// if the title and body are the same, wait 10 minutes
+				if ((deltaTime.TotalSeconds < 600) && title == _lastTitle && body == _lastBody)
 				{
 					return;
 				}
@@ -102,7 +107,7 @@ namespace UnrealToolbox
 				Body = body
 			};
 
-			_platformManager?.ShowNotification(notification, DateTimeOffset.Now + TimeSpan.FromSeconds(30));
+			_platformManager.ShowNotification(notification, DateTimeOffset.Now + TimeSpan.FromSeconds(30));
 		}
 	}
 }
\ No newline at end of file
