Horde: Don't surface notifications when agent is disabled

diff --git a/Engine/Source/Programs/UnrealToolbox/Plugins/HordeAgent/HordeAgentPlugin.cs b/Engine/Source/Programs/UnrealToolbox/Plugins/HordeAgent/HordeAgentPlugin.cs
--- a/Engine/Source/Programs/UnrealToolbox/Plugins/HordeAgent/HordeAgentPlugin.cs
+++ b/Engine/Source/Programs/UnrealToolbox/Plugins/HordeAgent/HordeAgentPlugin.cs
@@ -254,6 +254,11 @@ namespace UnrealToolbox.Plugins.HordeAgent
 
 		public TrayAppPluginStatus GetStatus()
 		{
+			if (_settings.Mode == AgentMode.Disabled)
+			{
+				return TrayAppPluginStatus.Default;
+			}
+
 			if (!_toolCatalog.Items.Any(x => x.Pending != null && !x.Pending.Failed))
 			{
 				_reportStatus = _status;
