Horde: Notification improvements

diff --git a/Engine/Source/Programs/UnrealToolbox/App.axaml.cs b/Engine/Source/Programs/UnrealToolbox/App.axaml.cs
--- a/Engine/Source/Programs/UnrealToolbox/App.axaml.cs
+++ b/Engine/Source/Programs/UnrealToolbox/App.axaml.cs
@@ -46,9 +46,7 @@ namespace UnrealToolbox
 			serviceCollection.AddSingleton<IToolCatalog, ToolCatalog>(sp => sp.GetRequiredService<ToolCatalog>());
 			serviceCollection.AddSingleton<ITrayAppHost>(this);
 			serviceCollection.AddSingleton<SelfUpdateService>();
-
 			serviceCollection.AddSingleton<ToolboxNotificationManager>();
-			serviceCollection.AddSingleton<IToolboxNotificationManager, ToolboxNotificationManager>(sp => sp.GetRequiredService<ToolboxNotificationManager>());
 
 			// Configure plugins
 			serviceCollection.AddSingleton<ITrayAppPlugin, HordeAgentPlugin>();
@@ -320,17 +318,10 @@ namespace UnrealToolbox
 		{
 			TrayAppPluginState state = TrayAppPluginState.Undefined;
 			List<string> messages = new List<string>();
-			string? pluginName = null;
-			string? errorMessage = null;
 
 			foreach (ITrayAppPlugin plugin in _serviceProvider.GetServices<ITrayAppPlugin>())
 			{
 				TrayAppPluginStatus status = plugin.GetStatus();
-				if (status.State == TrayAppPluginState.Error && !String.IsNullOrEmpty(status.NotificationMessage))
-				{
-					pluginName = plugin.Name;
-					errorMessage = status.NotificationMessage;
-				}
 
 				if (status.State >= state)
 				{
@@ -355,16 +346,6 @@ namespace UnrealToolbox
 				_ => _normalIcon
 			};
 			trayIcon.ToolTipText = String.Join("\n", messages);
-
-			// Notify on errors
-			if (!String.IsNullOrEmpty(errorMessage))
-			{
-				IToolboxNotificationManager? notifications = _serviceProvider.GetService<IToolboxNotificationManager>();
-				if (notifications != null)
-				{
-					notifications.ShowNotification(pluginName ?? "Status Error", errorMessage);
-				}
-			}
 		}
 	}
 }
diff --git a/Engine/Source/Programs/UnrealToolbox/GeneralSettingsViewModel.cs b/Engine/Source/Programs/UnrealToolbox/GeneralSettingsViewModel.cs
--- a/Engine/Source/Programs/UnrealToolbox/GeneralSettingsViewModel.cs
+++ b/Engine/Source/Programs/UnrealToolbox/GeneralSettingsViewModel.cs
@@ -158,10 +158,12 @@ namespace UnrealToolbox
 				bool result = await client.LoginAsync(true, cancellationToken);
 				if (!result)
 				{
+					ToolboxNotificationManager.PostNotification($"Connection Error: {client.ServerUrl.ToString()}", "Login Failed");
 					return new ConnectionState(client.ServerUrl.ToString(), "Login failed");
 				}
 				else if (!client.HasValidAccessToken())
 				{
+					ToolboxNotificationManager.PostNotification($"Connection Error: {client.ServerUrl.ToString()}", "Session expired");
 					return new ConnectionState(client.ServerUrl.ToString(), "Session expired");
 				}
 				else
@@ -171,6 +173,18 @@ namespace UnrealToolbox
 			}
 			catch (Exception ex)
 			{
+				string message = ex.Message;
+				message = message.Length > 100 ? message.Substring(0, 100) : message;
+
+				// Make some connection errors more friendly
+				string serverUrl = HordeOptions.GetDefaultServerUrl()?.ToString() ?? "Not Conigured";
+				if (message.Contains("party did not properly respond", StringComparison.OrdinalIgnoreCase) || message.Contains("actively refused it", StringComparison.OrdinalIgnoreCase))
+				{
+					message = $"Unable to reach the Horde Server: {serverUrl}";
+				}
+
+				ToolboxNotificationManager.PostNotification($"Connection Error", $"{message}");
+
 				return new ConnectionState(hordeClientRef.Client.ServerUrl.ToString(), $"Connection failed: {ex.Message}");
 			}
 		}
diff --git a/Engine/Source/Programs/UnrealToolbox/IToolboxPlugin.cs b/Engine/Source/Programs/UnrealToolbox/IToolboxPlugin.cs
--- a/Engine/Source/Programs/UnrealToolbox/IToolboxPlugin.cs
+++ b/Engine/Source/Programs/UnrealToolbox/IToolboxPlugin.cs
@@ -17,7 +17,7 @@ namespace UnrealToolbox
 	/// <summary>
 	/// Reported status of a plugin
 	/// </summary>
-	record class TrayAppPluginStatus(TrayAppPluginState State, string? Message = null, string? NotificationMessage = null)
+	record class TrayAppPluginStatus(TrayAppPluginState State, string? Message = null)
 	{
 		public static TrayAppPluginStatus Default { get; } = new TrayAppPluginStatus(TrayAppPluginState.Undefined);
 	}
diff --git a/Engine/Source/Programs/UnrealToolbox/Notifications/Notifications.cs b/Engine/Source/Programs/UnrealToolbox/Notifications/Notifications.cs
--- a/Engine/Source/Programs/UnrealToolbox/Notifications/Notifications.cs
+++ b/Engine/Source/Programs/UnrealToolbox/Notifications/Notifications.cs
@@ -4,30 +4,56 @@ using System.Reflection;
 using System.Runtime.InteropServices;
 using DesktopNotifications;
 using DesktopNotifications.Windows;
+using EpicGames.Core;
 
 namespace UnrealToolbox
 {
 	/// <summary>
-	/// Toolbox notification manager interface
+	/// A tool notification
 	/// </summary>
-	public interface IToolboxNotificationManager
+	class ToolboxNotification
 	{
 		/// <summary>
-		/// Show a notification
+		/// The title to display
+		/// </summary>
+		public string Title { get; }
+
+		/// <summary>
+		/// The body of the notification
+		/// </summary>
+		public string Body { get; }
+
+		/// <summary>
+		/// Whether to force the notification regardless of delta time
+		/// </summary>
+		public bool Force { get; }
+
+		/// <summary>
+		/// Constructor
 		/// </summary>
 		/// <param name="title"></param>
 		/// <param name="body"></param>
 		/// <param name="force"></param>
-		void ShowNotification(string title, string body, bool force = false);
+		public ToolboxNotification(string title, string body, bool force = false)
+		{
+			Title = title;
+			Body = body;
+			Force = force;
+		}		
 	}
 
 	/// <summary>
 	/// Toolbox notification manager implementation
 	/// </summary>
-	class ToolboxNotificationManager : IToolboxNotificationManager, IDisposable
+	class ToolboxNotificationManager : IAsyncDisposable
 	{
 		private INotificationManager? _platformManager;
 
+		private static readonly object s_lock = new object();
+
+		private static List<ToolboxNotification> s_notifications = new List<ToolboxNotification>();
+		readonly BackgroundTask _backgroundTask;
+
 		// spam prevention
 		DateTime? _lastNotificationTime;
 		string? _lastTitle;
@@ -38,8 +64,45 @@ namespace UnrealToolbox
 
 		public ToolboxNotificationManager()
 		{
+			_backgroundTask = new BackgroundTask(PostNotificationsAsync);
+		}
+
+		async Task PostNotificationsAsync(CancellationToken cancellationToken)
+		{
+			for (; ; )
+			{
+				try
+				{
+					List<ToolboxNotification>? notifications = null;
+					lock (s_lock)
+					{
+						if (s_notifications.Count > 0)
+						{
+							notifications = s_notifications;
+							s_notifications = new List<ToolboxNotification>();
+						}	
+					}
+
+					if (notifications != null && notifications.Count > 0)
+					{
+						notifications.Reverse();
+						ShowNotification(notifications[0].Title, notifications[0].Body, notifications[0].Force);
+					}
+				}
+				catch (OperationCanceledException)
+				{
+					throw;
+				}
+				catch (Exception)
+				{
+					
+				}
+
+				await Task.Delay(TimeSpan.FromSeconds(1.0), cancellationToken);
+			}
 		}
 
+
 		public void Start()
 		{
 			if (Environment.OSVersion.Platform == PlatformID.Win32NT)
@@ -59,11 +122,28 @@ namespace UnrealToolbox
 			// initialize and ensure with result
 			_platformManager.Initialize().GetAwaiter().GetResult();
 
+			_backgroundTask.Start();
+
 		}
 
-		public void Dispose()
+		public async ValueTask DisposeAsync()
 		{
 			_platformManager?.Dispose();
+			await _backgroundTask.DisposeAsync();
+		}
+
+		/// <summary>
+		/// Threead safe notification posting
+		/// </summary>
+		/// <param name="title"></param>
+		/// <param name="body"></param>
+		/// <param name="force"></param>
+		public static void PostNotification(string title, string body, bool force = false)
+		{
+			lock (s_lock)
+			{
+				s_notifications.Add(new ToolboxNotification(title, body, force));
+			}
 		}
 
 		/// <summary>
@@ -72,7 +152,7 @@ namespace UnrealToolbox
 		/// <param name="title"></param>
 		/// <param name="body"></param>
 		/// <param name="force"></param>
-		public void ShowNotification(string title, string body, bool force = false)
+		private void ShowNotification(string title, string body, bool force = false)
 		{
 			if (_platformManager == null)
 			{
diff --git a/Engine/Source/Programs/UnrealToolbox/Plugins/HordeAgent/HordeAgentPlugin.cs b/Engine/Source/Programs/UnrealToolbox/Plugins/HordeAgent/HordeAgentPlugin.cs
--- a/Engine/Source/Programs/UnrealToolbox/Plugins/HordeAgent/HordeAgentPlugin.cs
+++ b/Engine/Source/Programs/UnrealToolbox/Plugins/HordeAgent/HordeAgentPlugin.cs
@@ -2,6 +2,7 @@
 
 using Avalonia.Controls;
 using EpicGames.Core;
+using EpicGames.Horde;
 using FluentAvalonia.UI.Controls;
 using HordeAgent;
 using Microsoft.Extensions.Logging;
@@ -230,9 +231,25 @@ namespace UnrealToolbox.Plugins.HordeAgent
 				_status = null;
 			}
 			else if (!status.Healthy)
-			{
+			{				
 				string message = String.IsNullOrEmpty(status.Detail) ? "Error. Check logs." : status.Detail.Length > 100 ? status.Detail.Substring(0, 100) : status.Detail;
-				_status = new TrayAppPluginStatus(TrayAppPluginState.Error, message, message);
+				_status = new TrayAppPluginStatus(TrayAppPluginState.Error, message);
+				
+				if (_settings.Mode != AgentMode.Disabled)
+				{
+					// Make some known error messages more friendly
+					if (message.Contains("actively refused it", StringComparison.OrdinalIgnoreCase))
+					{
+						message = $"Could not connect to Horde Server: {HordeOptions.GetDefaultServerUrl()?.ToString() ?? "Not Conigured"}";
+					}
+
+					if (message.Contains("enrollment key does not match", StringComparison.OrdinalIgnoreCase))
+					{
+						message = $"Agent registration revoked by Horde Server: {HordeOptions.GetDefaultServerUrl()?.ToString() ?? "Not Conigured"}";
+					}
+
+					ToolboxNotificationManager.PostNotification("Horde Agent", message);
+				}
 			}
 			else if (status.NumLeases > 0)
 			{
