Branch snapshot for CL 37323767

[CL 37323767 in ue5-main branch]

diff --git a/Engine/Source/Programs/UnrealToolbox/App.axaml b/Engine/Source/Programs/UnrealToolbox/App.axaml
new file mode 100644
--- /dev/null
+++ b/Engine/Source/Programs/UnrealToolbox/App.axaml
@@ -0,0 +1,26 @@
+<Application xmlns="[REDACTED_PROFILE_URL]"
+			xmlns:x="http://schemas.microsoft.com/winfx/2006/xaml"
+			xmlns:sty="using:FluentAvalonia.Styling"
+			x:Class="UnrealToolbox.App"
+			RequestedThemeVariant="Default">
+			<!-- "Default" ThemeVariant follows system theme variant. "Dark" or "Light" are other available options. -->
+	
+	<Application.Styles>
+		<!-- As of preview5, these properties now default to false -->
+		<sty:FluentAvaloniaTheme PreferSystemTheme="True" PreferUserAccentColor="True" />
+	</Application.Styles>
+
+	<Application.Resources>
+		<WindowIcon x:Key="StatusBusy">/Resources/StatusBusy.ico</WindowIcon>
+		<WindowIcon x:Key="StatusError">/Resources/StatusError.ico</WindowIcon>
+		<WindowIcon x:Key="StatusNormal">/Resources/StatusNormal.ico</WindowIcon>
+		<WindowIcon x:Key="StatusPaused">/Resources/StatusPaused.ico</WindowIcon>
+	</Application.Resources>
+	
+	<TrayIcon.Icons>
+		<TrayIcons>
+			<TrayIcon Clicked="TrayIcon_Click"/>
+		</TrayIcons>
+	</TrayIcon.Icons>
+	
+</Application>
\ No newline at end of file
diff --git a/Engine/Source/Programs/UnrealToolbox/Plugins/HordeProxy/HordeProxySettingsPage.axaml b/Engine/Source/Programs/UnrealToolbox/Plugins/HordeProxy/HordeProxySettingsPage.axaml
new file mode 100644
--- /dev/null
+++ b/Engine/Source/Programs/UnrealToolbox/Plugins/HordeProxy/HordeProxySettingsPage.axaml
@@ -0,0 +1,56 @@
+<UserControl xmlns="[REDACTED_PROFILE_URL]"
+		xmlns:x="http://schemas.microsoft.com/winfx/2006/xaml"
+		xmlns:d="http://schemas.microsoft.com/expression/blend/2008"
+		xmlns:mc="http://schemas.openxmlformats.org/markup-compatibility/2006"
+		xmlns:local="using:UnrealToolbox.Plugins.HordeProxy"
+		mc:Ignorable="d" d:DesignWidth="800" d:DesignHeight="1000"
+		x:Class="UnrealToolbox.Plugins.HordeProxy.HordeProxySettingsPage"
+		x:DataType="local:HordeProxySettingsViewModel"
+		Background="{DynamicResource SolidBackgroundFillColorBaseBrush}">
+
+	<UserControl.Styles>
+		<StyleInclude Source="/Styles.axaml"/>
+	</UserControl.Styles>
+
+	<ScrollViewer>
+		<StackPanel Margin="20" Spacing="5">
+
+			<TextBlock Theme="{StaticResource SubtitleTextBlockStyle}" Margin="0 0 0 0">Horde Proxy</TextBlock>
+
+			<TextBlock Classes="SettingHeader">General</TextBlock>
+
+			<Border Classes="Card" Padding="16 6 16 2">
+				<StackPanel Spacing="10" Margin="0 0 0 5">
+					<Grid ColumnDefinitions="*,Auto">
+						<StackPanel Grid.Column="0" Margin="0 10 20 10">
+							<TextBlock Classes="SettingCaption">Enable Proxy Server</TextBlock>
+							<TextBlock Classes="SettingHelp">Provides an unauthenticated connection to via localhost. Useful for scripting or accessing Horde's symbol server through Visual Studio without applications having to handle authentication.</TextBlock>
+						</StackPanel>
+						<ToggleSwitch Grid.Column="1" VerticalAlignment="Center" OnContent="" OffContent="" IsChecked="{Binding Enabled}"/>
+					</Grid>
+				</StackPanel>
+			</Border>
+
+			<TextBlock Classes="SettingHeader" IsVisible="{Binding Enabled}">Network</TextBlock>
+
+			<Border Classes="Card" IsVisible="{Binding Enabled}">
+				<Grid ColumnDefinitions="*,Auto">
+					<StackPanel Grid.Column="0">
+						<TextBlock Classes="SettingCaption">Port</TextBlock>
+						<TextBlock Classes="SettingHelp">Local port number to allow connecting to the Horde Server</TextBlock>
+					</StackPanel>
+					<NumericUpDown VerticalAlignment="Center" Grid.Column="1" ShowButtonSpinner="False" Minimum="1" Maximum="65536" FormatString="0" TextAlignment="Right" Width="150" Value="{Binding Port}"/>
+				</Grid>
+			</Border>
+
+			<TextBlock Classes="SettingHeader" IsVisible="{Binding Enabled}">Status</TextBlock>
+
+			<Border Classes="Card" IsVisible="{Binding Enabled}">
+				<Grid ColumnDefinitions="*,Auto">
+					<TextBlock Classes="SettingHelp" Text="{Binding Status}"/>
+				</Grid>
+			</Border>
+
+		</StackPanel>
+	</ScrollViewer>
+</UserControl>
diff --git a/Engine/Source/Programs/UnrealToolbox/SettingsWindow.axaml b/Engine/Source/Programs/UnrealToolbox/SettingsWindow.axaml
new file mode 100644
--- /dev/null
+++ b/Engine/Source/Programs/UnrealToolbox/SettingsWindow.axaml
@@ -0,0 +1,21 @@
+<Window xmlns="[REDACTED_PROFILE_URL]"
+		xmlns:x="http://schemas.microsoft.com/winfx/2006/xaml"
+		xmlns:d="http://schemas.microsoft.com/expression/blend/2008"
+		xmlns:mc="http://schemas.openxmlformats.org/markup-compatibility/2006"
+		xmlns:pages="using:UnrealToolbox"
+		xmlns:ui="using:FluentAvalonia.UI.Controls"
+		mc:Ignorable="d" d:DesignWidth="800" d:DesignHeight="450"
+		x:Class="UnrealToolbox.SettingsWindow"
+		Width="1000"
+		Height="600"
+		WindowStartupLocation="CenterScreen" 
+		ExtendClientAreaChromeHints="NoChrome"
+		Title="Settings"
+		Icon="/Resources/StatusNormal.ico">
+
+	<Panel>
+		<ui:NavigationView x:Name="_navView" OpenPaneLength="260" IsSettingsVisible="False" PaneDisplayMode="Left" PaneTitle="Unreal Toolbox" BorderBrush="Transparent">
+		</ui:NavigationView>
+	</Panel>
+
+</Window>
diff --git a/Engine/Source/Programs/UnrealToolbox/Styles.axaml b/Engine/Source/Programs/UnrealToolbox/Styles.axaml
new file mode 100644
--- /dev/null
+++ b/Engine/Source/Programs/UnrealToolbox/Styles.axaml
@@ -0,0 +1,38 @@
+<Styles xmlns="[REDACTED_PROFILE_URL]"
+		xmlns:x="http://schemas.microsoft.com/winfx/2006/xaml">
+
+	<Design.PreviewWith>
+		<Panel Margin="50">
+			<Border Classes="Card" Padding="20">
+				<StackPanel>
+					<TextBlock Classes="SettingCaption">This is a setting</TextBlock>
+					<TextBlock Classes="SettingHelp">This is what the setting does</TextBlock>
+				</StackPanel>
+			</Border>
+		</Panel>
+	</Design.PreviewWith>
+
+	<Style Selector="Border.Card">
+		<Setter Property="BorderBrush" Value="{DynamicResource CardStrokeColorDefaultBrush}"/>
+		<Setter Property="BorderThickness" Value="1"/>
+		<Setter Property="Background" Value="{DynamicResource CardBackgroundFillColorDefaultBrush}"/>
+		<Setter Property="TextElement.Foreground" Value="{DynamicResource TextFillColorPrimaryBrush}"/>
+		<Setter Property="CornerRadius" Value="{DynamicResource ControlCornerRadius}"/>
+		<Setter Property="BackgroundSizing" Value="InnerBorderEdge"/>
+		<Setter Property="Padding" Value="16 14"/>
+	</Style>
+	<Style Selector="TextBlock.SettingHeader">
+		<Setter Property="Theme" Value="{DynamicResource BodyStrongTextBlockStyle}"/>
+		<Setter Property="Margin" Value="0 16 0 8"/>
+	</Style>
+	<Style Selector="TextBlock.SettingCaption">
+		<Setter Property="Theme" Value="{DynamicResource BodyTextBlockStyle}"/>
+		<Setter Property="VerticalAlignment" Value="Center"/>		
+	</Style>
+	<Style Selector="TextBlock.SettingHelp">
+		<Setter Property="Theme" Value="{DynamicResource BodyTextBlockStyle}"/>
+		<Setter Property="FontWeight" Value="Light"/>
+		<Setter Property="VerticalAlignment" Value="Center"/>
+	</Style>
+
+</Styles>
