Unreal Toolbox: Add CPU utilization settings for agent

Updates agent via IPC with target count and multiplier to use. Fix style warning in Notifications.cs.

diff --git a/Engine/Source/Programs/UnrealToolbox/Plugins/HordeAgent/HordeAgentPlugin.cs b/Engine/Source/Programs/UnrealToolbox/Plugins/HordeAgent/HordeAgentPlugin.cs
--- a/Engine/Source/Programs/UnrealToolbox/Plugins/HordeAgent/HordeAgentPlugin.cs
+++ b/Engine/Source/Programs/UnrealToolbox/Plugins/HordeAgent/HordeAgentPlugin.cs
@@ -240,12 +240,12 @@ namespace UnrealToolbox.Plugins.HordeAgent
 					// Make some known error messages more friendly
 					if (message.Contains("actively refused it", StringComparison.OrdinalIgnoreCase))
 					{
-						message = $"Could not connect to Horde Server: {HordeOptions.GetDefaultServerUrl()?.ToString() ?? "Not Conigured"}";
+						message = $"Could not connect to Horde Server: {HordeOptions.GetDefaultServerUrl()?.ToString() ?? "(Not configured)"}";
 					}
 
 					if (message.Contains("enrollment key does not match", StringComparison.OrdinalIgnoreCase))
 					{
-						message = $"Agent registration revoked by Horde Server: {HordeOptions.GetDefaultServerUrl()?.ToString() ?? "Not Conigured"}";
+						message = $"Agent registration revoked by Horde Server: {HordeOptions.GetDefaultServerUrl()?.ToString() ?? "(Not configured)"}";
 					}
 
 					ToolboxNotificationManager.PostNotification("Horde Agent", message);
@@ -548,6 +548,9 @@ namespace UnrealToolbox.Plugins.HordeAgent
 							_agentSettings = message.Parse<AgentSettingsMessage>();
 						}
 					}
+					
+					message.Set(AgentMessageType.SetSettingsRequest, new AgentSetSettingsRequest(_settings.Cpu.CpuCount, _settings.Cpu.CpuMultiplier));
+					await message.SendAsync(pipeClient, cancellationToken);
 
 					message.Set(AgentMessageType.GetStatusRequest);
 					await message.SendAsync(pipeClient, cancellationToken);
diff --git a/Engine/Source/Programs/UnrealToolbox/Plugins/HordeAgent/HordeAgentSettings.cs b/Engine/Source/Programs/UnrealToolbox/Plugins/HordeAgent/HordeAgentSettings.cs
--- a/Engine/Source/Programs/UnrealToolbox/Plugins/HordeAgent/HordeAgentSettings.cs
+++ b/Engine/Source/Programs/UnrealToolbox/Plugins/HordeAgent/HordeAgentSettings.cs
@@ -16,6 +16,7 @@ namespace UnrealToolbox.Plugins.HordeAgent
 	{
 		public AgentMode? Mode { get; set; }
 		public IdleSettings Idle { get; set; } = new();
+		public CpuUtilizationSettings Cpu { get; set; } = new();
 	}
 
 	[Serializable]
@@ -27,4 +28,11 @@ namespace UnrealToolbox.Plugins.HordeAgent
 
 		public string[] CriticalProcesses { get; set; } = Array.Empty<string>();
 	}
+	
+	[Serializable]
+	public class CpuUtilizationSettings
+	{
+		public int CpuCount { get; set; } = Environment.ProcessorCount;
+		public double CpuMultiplier { get; set; } = 1.0;
+	}
 }
diff --git a/Engine/Source/Programs/UnrealToolbox/Plugins/HordeAgent/HordeAgentSettingsPage.axaml b/Engine/Source/Programs/UnrealToolbox/Plugins/HordeAgent/HordeAgentSettingsPage.axaml
--- a/Engine/Source/Programs/UnrealToolbox/Plugins/HordeAgent/HordeAgentSettingsPage.axaml
+++ b/Engine/Source/Programs/UnrealToolbox/Plugins/HordeAgent/HordeAgentSettingsPage.axaml
@@ -85,6 +85,28 @@
 				</Grid>
 			</Border>
 
+			<TextBlock Classes="SettingHeader" IsVisible="{Binding IsWorkstationMode}">CPU Utilization</TextBlock>
+
+			<Border Classes="Card" IsVisible="{Binding IsWorkstationMode}">
+				<Grid ColumnDefinitions="*,Auto">
+					<StackPanel Grid.Column="0">
+						<TextBlock Classes="SettingCaption">CPU Count</TextBlock>
+						<TextBlock Classes="SettingHelp">Maximum number of logical CPU cores workloads should use</TextBlock>
+					</StackPanel>
+					<NumericUpDown x:Name="CpuCount" VerticalAlignment="Center" Grid.Column="1" Minimum="0" Maximum="512" FormatString="0 cores" Width="160" Value="{Binding CpuCount}"/>
+				</Grid>
+			</Border>
+
+			<Border Classes="Card" IsVisible="{Binding IsWorkstationMode}">
+				<Grid ColumnDefinitions="*,Auto">
+					<StackPanel Grid.Column="0">
+						<TextBlock Classes="SettingCaption">CPU Multiplier</TextBlock>
+						<TextBlock Classes="SettingHelp">Multiplier applied to CPU core count</TextBlock>
+					</StackPanel>
+					<NumericUpDown Grid.Column="1" Minimum="0" Maximum="20" Increment="0.1" FormatString="0.0x" Width="160" VerticalAlignment="Center" Value="{Binding CpuMultiplier}"/>
+				</Grid>
+			</Border>
+
 		</StackPanel>
 	</ScrollViewer>
 </UserControl>
diff --git a/Engine/Source/Programs/UnrealToolbox/Plugins/HordeAgent/HordeAgentSettingsPage.axaml.cs b/Engine/Source/Programs/UnrealToolbox/Plugins/HordeAgent/HordeAgentSettingsPage.axaml.cs
--- a/Engine/Source/Programs/UnrealToolbox/Plugins/HordeAgent/HordeAgentSettingsPage.axaml.cs
+++ b/Engine/Source/Programs/UnrealToolbox/Plugins/HordeAgent/HordeAgentSettingsPage.axaml.cs
@@ -16,6 +16,10 @@ namespace UnrealToolbox.Plugins.HordeAgent
 			InitializeComponent();
 
 			DataContext = new HordeAgentSettingsViewModel(context, plugin);
+			
+			// Multiply with two as ProcessorCount may not return all logical processors depending on system
+			// TODO: Replace with more accurate lookup (requires more .NET deps)
+			CpuCount.Maximum = Environment.ProcessorCount * 2;
 		}
 	}
 }
\ No newline at end of file
diff --git a/Engine/Source/Programs/UnrealToolbox/Plugins/HordeAgent/HordeAgentSettingsViewModel.cs b/Engine/Source/Programs/UnrealToolbox/Plugins/HordeAgent/HordeAgentSettingsViewModel.cs
--- a/Engine/Source/Programs/UnrealToolbox/Plugins/HordeAgent/HordeAgentSettingsViewModel.cs
+++ b/Engine/Source/Programs/UnrealToolbox/Plugins/HordeAgent/HordeAgentSettingsViewModel.cs
@@ -65,6 +65,12 @@ namespace UnrealToolbox.Plugins.HordeAgent
 
 		[ObservableProperty]
 		int _minFreeCpuPct;
+		
+		[ObservableProperty]
+		int _cpuCount;
+		
+		[ObservableProperty]
+		double _cpuMultiplier;
 
 		[ObservableProperty]
 		string _sandboxDir = "C:\\Dump"; // TODO
@@ -79,6 +85,8 @@ namespace UnrealToolbox.Plugins.HordeAgent
 			_minIdleTimeSecs = settings.Idle.MinIdleTimeSecs;
 			_minFreeMemoryMb = settings.Idle.MinFreeVirtualMemMb;
 			_minFreeCpuPct = settings.Idle.MinIdleCpuPct;
+			_cpuCount = settings.Cpu.CpuCount;
+			_cpuMultiplier = settings.Cpu.CpuMultiplier;
 		}
 
 		protected override void OnPropertyChanged(PropertyChangedEventArgs e)
@@ -93,6 +101,8 @@ namespace UnrealToolbox.Plugins.HordeAgent
 				settings.Idle.MinIdleTimeSecs = MinIdleTimeSecs;
 				settings.Idle.MinFreeVirtualMemMb = MinFreeMemoryMb;
 				settings.Idle.MinIdleCpuPct = MinFreeCpuPct;
+				settings.Cpu.CpuCount = CpuCount;
+				settings.Cpu.CpuMultiplier = CpuMultiplier;
 
 				_plugin.UpdateSettings(settings);
 			}
