Unreal Toolbox: Fix connection errors causing overlapping text in settings window.

#rnx

diff --git a/Engine/Source/Programs/UnrealToolbox/GeneralSettingsPage.axaml b/Engine/Source/Programs/UnrealToolbox/GeneralSettingsPage.axaml
--- a/Engine/Source/Programs/UnrealToolbox/GeneralSettingsPage.axaml
+++ b/Engine/Source/Programs/UnrealToolbox/GeneralSettingsPage.axaml
@@ -21,18 +21,17 @@
 			<TextBlock Classes="SettingHeader">Horde Server</TextBlock>
 
 			<Border Classes="Card">
-				<Grid ColumnDefinitions="*,Auto,Auto,Auto" RowDefinitions="Auto">
-					<StackPanel Grid.Column="0" Margin="0">
-						<HyperlinkButton Grid.Column="0" Grid.Row="0" Padding="0" FontSize="14" Content="{Binding ServerUrl}" Command="{Binding OpenServer}"/>
-						<Grid ColumnDefinitions="Auto,Auto" RowDefinitions="Auto">
-							<ui:ProgressRing Grid.Column="0" Grid.Row="0" Width="16" Height="16" Margin="0 0 6 0" IsVisible="{Binding IsConnecting}" VerticalAlignment="Center"/>
-							<TextBlock Grid.Column="1" Grid.Row="0" Padding="1" FontSize="12" Text="{Binding ServerStatus}" VerticalAlignment="Top" TextTrimming="CharacterEllipsis"/>
-						</Grid>
-					</StackPanel>
+				<Grid ColumnDefinitions="*,Auto,Auto,Auto" RowDefinitions="Auto,Auto">
+					<HyperlinkButton Grid.Column="0" Grid.Row="0" Padding="0" FontSize="14" Content="{Binding ServerUrl}" Command="{Binding OpenServer}"/>
+
+					<Grid Grid.Column="0" Grid.Row="1" ColumnDefinitions="Auto,*" RowDefinitions="Auto">
+						<ui:ProgressRing Grid.Column="0" Grid.Row="0" Width="16" Height="16" Margin="0 2 6 2" IsVisible="{Binding IsConnecting}" VerticalAlignment="Top"/>
+						<TextBlock Grid.Column="1" Grid.Row="0" Padding="1" FontSize="12" Text="{Binding ServerStatus}" VerticalAlignment="Top" TextTrimming="CharacterEllipsis" TextWrapping="Wrap"/>
+					</Grid>
 
-					<Button Grid.Column="1" Command="{Binding StartRefresh}" IsEnabled="{Binding !IsConnecting}" Margin="0 0 5 0">Refresh</Button>
-					<Button Grid.Column="2" Command="{Binding CancelRefresh}" IsEnabled="{Binding IsConnecting}" Margin="0 0 5 0">Cancel</Button>
-					<Button Grid.Column="3" Command="{Binding ConfigureServerAsync}">Modify...</Button>
+					<Button Grid.Column="1" Command="{Binding StartRefresh}" IsEnabled="{Binding !IsConnecting}" Margin="0 0 5 0" Grid.RowSpan="2">Refresh</Button>
+					<Button Grid.Column="2" Command="{Binding CancelRefresh}" IsEnabled="{Binding IsConnecting}" Margin="0 0 5 0" Grid.RowSpan="2">Cancel</Button>
+					<Button Grid.Column="3" Command="{Binding ConfigureServerAsync}" Grid.RowSpan="2">Modify...</Button>
 				</Grid>
 			</Border>
 
