Unreal Toolbox: Link settings directory and log file locations from about page

diff --git a/Engine/Source/Programs/UnrealToolbox/AboutPage.axaml b/Engine/Source/Programs/UnrealToolbox/AboutPage.axaml
--- a/Engine/Source/Programs/UnrealToolbox/AboutPage.axaml
+++ b/Engine/Source/Programs/UnrealToolbox/AboutPage.axaml
@@ -1,31 +1,47 @@
 <UserControl xmlns="[REDACTED_PROFILE_URL]"
-			xmlns:x="http://schemas.microsoft.com/winfx/2006/xaml"
-			xmlns:d="http://schemas.microsoft.com/expression/blend/2008"
-			xmlns:mc="http://schemas.openxmlformats.org/markup-compatibility/2006"
-			xmlns:ui="using:FluentAvalonia.UI.Controls"
-			xmlns:local="using:UnrealToolbox"
-			mc:Ignorable="d" d:DesignWidth="800" d:DesignHeight="600"
-			x:Class="UnrealToolbox.AboutPage"
-			Background="{DynamicResource SolidBackgroundFillColorBaseBrush}">
+             xmlns:x="http://schemas.microsoft.com/winfx/2006/xaml"
+             xmlns:d="http://schemas.microsoft.com/expression/blend/2008"
+             xmlns:mc="http://schemas.openxmlformats.org/markup-compatibility/2006"
+             xmlns:ui="using:FluentAvalonia.UI.Controls"
+             xmlns:local="using:UnrealToolbox"
+             mc:Ignorable="d" d:DesignWidth="800" d:DesignHeight="600"
+             x:Class="UnrealToolbox.AboutPage"
+             Background="{DynamicResource SolidBackgroundFillColorBaseBrush}">
 
 	<UserControl.Styles>
-		<StyleInclude Source="/Styles.axaml"/>
+		<StyleInclude Source="/Styles.axaml" />
 	</UserControl.Styles>
-	
+
 	<ScrollViewer>
 		<StackPanel Margin="20" Spacing="5">
-
 			<TextBlock Theme="{StaticResource SubtitleTextBlockStyle}" VerticalAlignment="Center" Margin="0 0 0 0">About</TextBlock>
 
-			<TextBlock Classes="SettingHeader">Version</TextBlock>
+			<Grid Margin="20">
+				<Grid.ColumnDefinitions>
+					<ColumnDefinition Width="Auto" MinWidth="180" />
+					<ColumnDefinition Width="*" />
+				</Grid.ColumnDefinitions>
+				<Grid.RowDefinitions>
+					<RowDefinition Height="Auto" />
+					<RowDefinition Height="Auto" />
+					<RowDefinition Height="Auto" />
+				</Grid.RowDefinitions>
 
-			<Border Classes="Card" Padding="16 6 16 2">
-				<StackPanel Spacing="10" Margin="0 10 20 15">
-					<TextBlock Classes="SettingHelp" Name="_versionText"/>
-				</StackPanel>
-			</Border>
+				<!-- Labels column -->
+				<TextBlock Grid.Row="0" Grid.Column="0" Text="Version:" VerticalAlignment="Center" Margin="0,0,10,10" />
+				<TextBlock Grid.Row="1" Grid.Column="0" Text="Settings and data directory:" VerticalAlignment="Center" Margin="0,0,10,10" />
+				<TextBlock Grid.Row="2" Grid.Column="0" Text="Log file:" VerticalAlignment="Center" Margin="0,0,10,10" />
 
+				<!-- Values column -->
+				<TextBlock Grid.Row="0" Grid.Column="1" Name="_versionText" VerticalAlignment="Center" Margin="0,0,0,10" />
+				<Button Grid.Row="1" Grid.Column="1" Margin="0,0,0,10" Padding="0" Background="Transparent" BorderThickness="0" Name="_settingsDirButton" Cursor="Hand">
+					<TextBlock Name="_settingsDirText" Foreground="Blue" TextDecorations="Underline" />
+				</Button>
+				<Button Grid.Row="2" Grid.Column="1" Margin="0,0,0,10" Padding="0" Background="Transparent" BorderThickness="0" Name="_logFileButton" Cursor="Hand">
+					<TextBlock Name="_logFileText" Foreground="Blue" TextDecorations="Underline" />
+				</Button>
+			</Grid>
 		</StackPanel>
 	</ScrollViewer>
 
-</UserControl>
+</UserControl>
\ No newline at end of file
diff --git a/Engine/Source/Programs/UnrealToolbox/AboutPage.axaml.cs b/Engine/Source/Programs/UnrealToolbox/AboutPage.axaml.cs
--- a/Engine/Source/Programs/UnrealToolbox/AboutPage.axaml.cs
+++ b/Engine/Source/Programs/UnrealToolbox/AboutPage.axaml.cs
@@ -1,5 +1,7 @@
 // Copyright Epic Games, Inc. All Rights Reserved.
 
+using System.Diagnostics;
+using System.Runtime.InteropServices;
 using Avalonia.Controls;
 
 namespace UnrealToolbox
@@ -12,6 +14,49 @@ namespace UnrealToolbox
 
 			string version = Program.Update?.CurrentVersionString ?? "No version information present.";
 			_versionText.Text = version;
+			
+			string settingsDir = ToolCatalog.GetBaseDir().FullName;
+			_settingsDirButton.Click += (sender, args) => OpenFolder(settingsDir);
+			_settingsDirText.Text = settingsDir;
+			
+			if (Program.LogFile != null)
+			{
+				string logFile = Program.LogFile.FullName;
+				_logFileButton.Click += (sender, args) => OpenFile(logFile);
+				_logFileText.Text = logFile;
+			}
+		}
+		
+		private static void OpenFolder(string path)
+		{
+			if (RuntimeInformation.IsOSPlatform(OSPlatform.Windows))
+			{
+				Process.Start("explorer.exe", path);
+			}
+			else if (RuntimeInformation.IsOSPlatform(OSPlatform.Linux))
+			{
+				Process.Start("xdg-open", path);
+			}
+			else if (RuntimeInformation.IsOSPlatform(OSPlatform.OSX))
+			{
+				Process.Start("open", path);
+			}
+		}
+		
+		private static void OpenFile(string path)
+		{
+			if (RuntimeInformation.IsOSPlatform(OSPlatform.Windows))
+			{
+				Process.Start(new ProcessStartInfo(path) { UseShellExecute = true });
+			}
+			else if (RuntimeInformation.IsOSPlatform(OSPlatform.Linux))
+			{
+				Process.Start("xdg-open", path);
+			}
+			else if (RuntimeInformation.IsOSPlatform(OSPlatform.OSX))
+			{
+				Process.Start("open", path);
+			}
 		}
 	}
 }
\ No newline at end of file
diff --git a/Engine/Source/Programs/UnrealToolbox/ToolCatalog.cs b/Engine/Source/Programs/UnrealToolbox/ToolCatalog.cs
--- a/Engine/Source/Programs/UnrealToolbox/ToolCatalog.cs
+++ b/Engine/Source/Programs/UnrealToolbox/ToolCatalog.cs
@@ -195,9 +195,6 @@ namespace UnrealToolbox
 		readonly IHordeClientProvider _hordeClientProvider;
 		readonly ILogger _logger;
 		readonly object _workerThreadLockObject = new object();
-		readonly DirectoryReference _baseDir;
-		readonly DirectoryReference _toolsDir;
-		readonly FileReference _configFile;
 		readonly AsyncEvent _updateEvent;
 		readonly BackgroundTask _updateTask;
 		readonly AsyncEvent _cleanupEvent;
@@ -234,14 +231,8 @@ namespace UnrealToolbox
 
 		public ToolCatalog(IHordeClientProvider hordeClientProvider, ILogger<ToolCatalog> logger)
 		{
-			DirectoryReference? toolsDir = DirectoryReference.GetSpecialFolder(Environment.SpecialFolder.LocalApplicationData);
-			toolsDir ??= DirectoryReference.GetCurrentDirectory();
-
 			_hordeClientProvider = hordeClientProvider;
 			_logger = logger;
-			_baseDir = DirectoryReference.Combine(toolsDir, "Epic Games", "Unreal Toolbox");
-			_toolsDir = DirectoryReference.Combine(_baseDir, "Tools");
-			_configFile = FileReference.Combine(_baseDir, "Tools.json");
 			_updateEvent = new AsyncEvent();
 			_updateTask = new BackgroundTask(UpdateAsync);
 			_cleanupEvent = new AsyncEvent();
@@ -255,6 +246,16 @@ namespace UnrealToolbox
 
 			UpdateMainThread();
 		}
+		
+		public static DirectoryReference GetBaseDir()
+		{
+			DirectoryReference? dir = DirectoryReference.GetSpecialFolder(Environment.SpecialFolder.LocalApplicationData);
+			dir ??= DirectoryReference.GetCurrentDirectory();
+			return DirectoryReference.Combine(dir, "Epic Games", "Unreal Toolbox");
+		}
+		
+		public static DirectoryReference GetToolsDir() => DirectoryReference.Combine(GetBaseDir(), "Tools");
+		public static FileReference GetConfigFile() => FileReference.Combine(GetBaseDir(), "Tools.json");
 
 		void PostMainThreadUpdate()
 		{
@@ -346,9 +347,9 @@ namespace UnrealToolbox
 			_updateEvent.Pulse();
 		}
 
-		DirectoryReference GetToolDir(ToolId toolId, ToolDeploymentId deploymentId)
+		static DirectoryReference GetToolDir(ToolId toolId, ToolDeploymentId deploymentId)
 		{
-			return DirectoryReference.Combine(_toolsDir, toolId.ToString(), deploymentId.ToString());
+			return DirectoryReference.Combine(GetToolsDir(), toolId.ToString(), deploymentId.ToString());
 		}
 
 		ToolConfig LoadToolConfig(DirectoryReference toolDir)
@@ -378,7 +379,7 @@ namespace UnrealToolbox
 		{
 			Dictionary<ToolId, Item> items = new Dictionary<ToolId, Item>();
 
-			byte[]? data = FileTransaction.ReadAllBytes(_configFile);
+			byte[]? data = FileTransaction.ReadAllBytes(GetConfigFile());
 			if (data != null)
 			{
 				JsonState? state = JsonSerializer.Deserialize<JsonState>(data, s_jsonSerializerOptions);
@@ -422,8 +423,8 @@ namespace UnrealToolbox
 
 			byte[] data = JsonSerializer.SerializeToUtf8Bytes(state, s_jsonSerializerOptions);
 
-			DirectoryReference.CreateDirectory(_configFile.Directory);
-			await FileTransaction.WriteAllBytesAsync(_configFile, data, cancellationToken);
+			DirectoryReference.CreateDirectory(GetConfigFile().Directory);
+			await FileTransaction.WriteAllBytesAsync(GetConfigFile(), data, cancellationToken);
 		}
 
 		void CancelDeployment(Item item, ToolDeploymentInfo? deployment)
@@ -753,7 +754,7 @@ namespace UnrealToolbox
 			List<DirectoryInfo> deleteToolDirs = new List<DirectoryInfo>();
 			List<DirectoryInfo> deleteDeploymentDirs = new List<DirectoryInfo>();
 
-			DirectoryInfo baseDirInfo = _toolsDir.ToDirectoryInfo();
+			DirectoryInfo baseDirInfo = GetToolsDir().ToDirectoryInfo();
 			if (baseDirInfo.Exists)
 			{
 				foreach (DirectoryInfo toolDir in baseDirInfo.EnumerateDirectories())
