Unreal Toolbox: Fix log file not opening in tool list

Avalonia can't generate trampoline code for raw static methods leading to InvalidProgramException.
Also add link to log file when install fails.

diff --git a/Engine/Source/Programs/UnrealToolbox/ToolCatalog.cs b/Engine/Source/Programs/UnrealToolbox/ToolCatalog.cs
--- a/Engine/Source/Programs/UnrealToolbox/ToolCatalog.cs
+++ b/Engine/Source/Programs/UnrealToolbox/ToolCatalog.cs
@@ -510,7 +510,7 @@ namespace UnrealToolbox
 				int exitCode = await RunCommandAsync(item.Id.ToString(), toolConfig.InstallCommand, toolDir.FullName, cancellationToken);
 				if (exitCode != 0)
 				{
-					pending = new PendingToolDeploymentInfo(true, $"Installation failed ({exitCode}).", deploymentInfo);
+					pending = new PendingToolDeploymentInfo(true, $"Installation failed ({exitCode})", deploymentInfo, ShowLogLink: true);
 					UpdateStateAndNotify(item, state => state with { Pending = pending });
 					return;
 				}
diff --git a/Engine/Source/Programs/UnrealToolbox/ToolCatalogViewModel.cs b/Engine/Source/Programs/UnrealToolbox/ToolCatalogViewModel.cs
--- a/Engine/Source/Programs/UnrealToolbox/ToolCatalogViewModel.cs
+++ b/Engine/Source/Programs/UnrealToolbox/ToolCatalogViewModel.cs
@@ -3,6 +3,7 @@
 using System.Collections.ObjectModel;
 using System.ComponentModel;
 using System.Diagnostics;
+using System.Diagnostics.CodeAnalysis;
 using CommunityToolkit.Mvvm.ComponentModel;
 
 namespace UnrealToolbox
@@ -175,15 +176,13 @@ namespace UnrealToolbox
 			ShowUpdate = _item.Latest != null && _item.Current != null && _item.Latest.Id != _item.Current.Id;
 			ShowRemove = _item.Current != null;
 		}
-
-		public static void ShowLog()
+		
+		[SuppressMessage("Performance", "CA1822:Mark members as static", Justification = "Avalonia cannot call static methods as a binding")]
+		public void ShowLog()
 		{
 			if (OperatingSystem.IsWindows() && Program.LogFile != null)
 			{
-				ProcessStartInfo info = new ProcessStartInfo();
-				info.FileName = Program.LogFile.FullName;
-				info.UseShellExecute = true;
-				Process.Start(info);
+				Process.Start(new ProcessStartInfo(Program.LogFile.FullName) { UseShellExecute = true });
 			}
 		}
 
