Unreal Toolbox: Add support for elevating installer/uninstaller process

Uses process start from .NET SDK as ManagedProcess does not support elevation.

diff --git a/Engine/Source/Programs/UnrealToolbox/ToolCatalog.cs b/Engine/Source/Programs/UnrealToolbox/ToolCatalog.cs
--- a/Engine/Source/Programs/UnrealToolbox/ToolCatalog.cs
+++ b/Engine/Source/Programs/UnrealToolbox/ToolCatalog.cs
@@ -507,7 +507,7 @@ namespace UnrealToolbox
 				PendingToolDeploymentInfo pending = new PendingToolDeploymentInfo(false, "Running install actions...", deploymentInfo);
 				UpdateStateAndNotify(item, state => state with { Pending = pending });
 
-				int exitCode = await RunCommandAsync(item.Id.ToString(), toolConfig.InstallCommand, toolDir.FullName, cancellationToken);
+				int exitCode = await RunCommandAsync(item.Id.ToString(), toolConfig.InstallCommand, toolDir.FullName, toolConfig.RequiresElevation, cancellationToken);
 				if (exitCode != 0)
 				{
 					pending = new PendingToolDeploymentInfo(true, $"Installation failed ({exitCode})", deploymentInfo, ShowLogLink: true);
@@ -543,7 +543,7 @@ namespace UnrealToolbox
 					ToolCommand installCommand = new ToolCommand { FileName = "msiexec.exe" };
 					installCommand.Arguments = new List<string> { "/I", msiFile.FullName };
 
-					int exitCode = await RunCommandAsync(item.Id.ToString(), installCommand, toolDir.FullName, cancellationToken);
+					int exitCode = await RunCommandAsync(item.Id.ToString(), installCommand, toolDir.FullName, toolConfig.RequiresElevation, cancellationToken);
 					if (exitCode != 0 || !IsMsiInstalled(item._state.MsiProductId))
 					{
 						UpdateStateAndNotify(item, state => state with { Pending = null });
@@ -590,7 +590,7 @@ namespace UnrealToolbox
 					PendingToolDeploymentInfo? pending = new PendingToolDeploymentInfo(false, $"Removing {current.Version}", null);
 					UpdateStateAndNotify(item, state => state with { Pending = pending, Current = null });
 
-					await RunCommandAsync(item.Id.ToString(), current.Config.UninstallCommand, current.Dir.FullName, cancellationToken);
+					await RunCommandAsync(item.Id.ToString(), current.Config.UninstallCommand, current.Dir.FullName, current.Config.RequiresElevation, cancellationToken);
 				}
 				else if (OperatingSystem.IsWindows() && !String.IsNullOrEmpty(item._state.MsiProductId) && IsMsiInstalled(item._state.MsiProductId))
 				{
@@ -614,7 +614,7 @@ namespace UnrealToolbox
 						ToolCommand uninstallCommand = new ToolCommand { FileName = "msiexec.exe" };
 						uninstallCommand.Arguments = new List<string> { "/x", $"{{{item._state.MsiProductId}}}" };
 
-						int exitCode = await RunCommandAsync(item.Id.ToString(), uninstallCommand, Directory.GetCurrentDirectory(), cancellationToken);
+						int exitCode = await RunCommandAsync(item.Id.ToString(), uninstallCommand, Directory.GetCurrentDirectory(), current.Config.RequiresElevation, cancellationToken);
 						if (exitCode == 1602)
 						{
 							// User cancelled (https://learn.microsoft.com/en-us/windows/win32/msi/error-codes)
@@ -637,37 +637,68 @@ namespace UnrealToolbox
 			return true;
 		}
 
-		async Task<int> RunCommandAsync(string toolName, ToolCommand command, string workingDir, CancellationToken cancellationToken)
+		async Task<int> RunCommandAsync(string toolName, ToolCommand command, string workingDir, bool elevateProcess, CancellationToken cancellationToken)
 		{
 			string arguments = String.Empty;
 			if (command.Arguments != null)
 			{
 				arguments = CommandLineArguments.Join(command.Arguments);
 			}
-
+			
 			string fileName = Path.Combine(workingDir, command.FileName);
 			if (!File.Exists(fileName))
 			{
 				fileName = command.FileName;
 			}
-
+			
 			_logger.LogInformation("{ToolName}> Running {FileName} {Arguments}", toolName, CommandLineArguments.Quote(fileName), arguments);
-
-			using (ManagedProcessGroup newProcessGroup = new ManagedProcessGroup())
-			using (ManagedProcess newProcess = new ManagedProcess(newProcessGroup, fileName, arguments, workingDir, null, null, ProcessPriorityClass.Normal))
+			
+			using Process process = new Process();
+			process.StartInfo = new ProcessStartInfo
+			{
+				FileName = fileName,
+				Arguments = arguments,
+				WorkingDirectory = workingDir,
+			};
+			
+			if (elevateProcess)
+			{
+				// When elevating the process, IO cannot be redirected in this process
+				process.StartInfo.UseShellExecute = true;
+				process.StartInfo.Verb = "runas"; // "runas" triggers the UAC prompt
+				_logger.LogInformation("Tool install runs as an elevated process. To read process output, check installer's own logging mechanism.");
+			}
+			else
 			{
-				for (; ; )
+				process.StartInfo.RedirectStandardOutput = true;
+				process.StartInfo.RedirectStandardError = true;
+				process.StartInfo.UseShellExecute = false;
+				process.StartInfo.CreateNoWindow = true;
+				
+				DataReceivedEventHandler onDataReceived = (sender, args) =>
 				{
-					string? line = await newProcess.ReadLineAsync(cancellationToken);
-					if (line == null)
+					if (args.Data != null)
 					{
-						await newProcess.WaitForExitAsync(cancellationToken);
-						_logger.LogInformation("{ToolName}> Exit code {ExitCode}", toolName, newProcess.ExitCode);
-						return newProcess.ExitCode;
+						_logger.LogInformation("{ToolName}> {Line}", toolName, args.Data);
 					}
-					_logger.LogInformation("{ToolName}> {Line}", toolName, line);
-				}
+				};
+				
+				process.OutputDataReceived += onDataReceived;
+				process.ErrorDataReceived += onDataReceived;
+			}
+			
+			process.Start();
+			if (!elevateProcess)
+			{
+				process.BeginOutputReadLine();
+				process.BeginErrorReadLine();
 			}
+			
+			await process.WaitForExitAsync(cancellationToken);
+			int exitCode = process.ExitCode;
+			_logger.LogInformation("{ToolName}> Exit code {ExitCode}", toolName, exitCode);
+			
+			return exitCode;
 		}
 
 		void UpdateStateAndNotify(Item item, Func<ItemState, ItemState> update)
diff --git a/Engine/Source/Programs/UnrealToolbox/ToolConfig.cs b/Engine/Source/Programs/UnrealToolbox/ToolConfig.cs
--- a/Engine/Source/Programs/UnrealToolbox/ToolConfig.cs
+++ b/Engine/Source/Programs/UnrealToolbox/ToolConfig.cs
@@ -16,6 +16,11 @@ namespace UnrealToolbox
 		/// Whether the install/uninstall for this tool needs to be triggered manually
 		/// </summary>
 		public bool ManualInstall { get; set; }
+		
+		/// <summary>
+		/// Whether the install/uninstall for this tool needs administrator or root privileges
+		/// </summary>
+		public bool RequiresElevation { get; set; } = false;
 
 		/// <summary>
 		/// Command to run when installing
