Unreal Toolbox: Rename ITrayAppPlugin to IToolboxPlugin.

#rnx

diff --git a/Engine/Source/Programs/UnrealToolbox/App.axaml.cs b/Engine/Source/Programs/UnrealToolbox/App.axaml.cs
--- a/Engine/Source/Programs/UnrealToolbox/App.axaml.cs
+++ b/Engine/Source/Programs/UnrealToolbox/App.axaml.cs
@@ -18,7 +18,7 @@ namespace UnrealToolbox
 	/// <summary>
 	/// Main application class
 	/// </summary>
-	public sealed partial class App : Application, ITrayAppHost, IAsyncDisposable
+	public sealed partial class App : Application, IToolboxPluginHost, IAsyncDisposable
 	{
 		readonly ServiceProvider _serviceProvider;
 		readonly ILogger _logger;
@@ -44,13 +44,13 @@ namespace UnrealToolbox
 			serviceCollection.AddSingleton<IHordeClientProvider, HordeClientProvider>();
 			serviceCollection.AddSingleton<ToolCatalog>();
 			serviceCollection.AddSingleton<IToolCatalog, ToolCatalog>(sp => sp.GetRequiredService<ToolCatalog>());
-			serviceCollection.AddSingleton<ITrayAppHost>(this);
+			serviceCollection.AddSingleton<IToolboxPluginHost>(this);
 			serviceCollection.AddSingleton<SelfUpdateService>();
 			serviceCollection.AddSingleton<ToolboxNotificationManager>();
 
 			// Configure plugins
-			serviceCollection.AddSingleton<ITrayAppPlugin, HordeAgentPlugin>();
-			serviceCollection.AddSingleton<ITrayAppPlugin, HordeProxyPlugin>();
+			serviceCollection.AddSingleton<IToolboxPlugin, HordeAgentPlugin>();
+			serviceCollection.AddSingleton<IToolboxPlugin, HordeProxyPlugin>();
 
 			_serviceProvider = serviceCollection.BuildServiceProvider();
 			_logger = _serviceProvider.GetRequiredService<ILogger<App>>();
@@ -162,7 +162,7 @@ namespace UnrealToolbox
 			TrayIcon trayIcon = TrayIcon.GetIcons(this)![0];
 			trayIcon.Menu = contextMenu;
 
-			foreach (ITrayAppPlugin plugin in _serviceProvider.GetServices<ITrayAppPlugin>())
+			foreach (IToolboxPlugin plugin in _serviceProvider.GetServices<IToolboxPlugin>())
 			{
 				plugin.PopulateContextMenu(contextMenu);
 			}
@@ -238,7 +238,7 @@ namespace UnrealToolbox
 		private void RefreshPlugins()
 		{
 			bool update = false;
-			foreach (ITrayAppPlugin plugin in _serviceProvider.GetServices<ITrayAppPlugin>())
+			foreach (IToolboxPlugin plugin in _serviceProvider.GetServices<IToolboxPlugin>())
 			{
 				update |= plugin.Refresh();
 			}
@@ -319,9 +319,9 @@ namespace UnrealToolbox
 			TrayAppPluginState state = TrayAppPluginState.Undefined;
 			List<string> messages = new List<string>();
 
-			foreach (ITrayAppPlugin plugin in _serviceProvider.GetServices<ITrayAppPlugin>())
+			foreach (IToolboxPlugin plugin in _serviceProvider.GetServices<IToolboxPlugin>())
 			{
-				TrayAppPluginStatus status = plugin.GetStatus();
+				ToolboxPluginStatus status = plugin.GetStatus();
 
 				if (status.State >= state)
 				{
diff --git a/Engine/Source/Programs/UnrealToolbox/IToolboxPlugin.cs b/Engine/Source/Programs/UnrealToolbox/IToolboxPlugin.cs
--- a/Engine/Source/Programs/UnrealToolbox/IToolboxPlugin.cs
+++ b/Engine/Source/Programs/UnrealToolbox/IToolboxPlugin.cs
@@ -17,15 +17,15 @@ namespace UnrealToolbox
 	/// <summary>
 	/// Reported status of a plugin
 	/// </summary>
-	record class TrayAppPluginStatus(TrayAppPluginState State, string? Message = null)
+	record class ToolboxPluginStatus(TrayAppPluginState State, string? Message = null)
 	{
-		public static TrayAppPluginStatus Default { get; } = new TrayAppPluginStatus(TrayAppPluginState.Undefined);
+		public static ToolboxPluginStatus Default { get; } = new ToolboxPluginStatus(TrayAppPluginState.Undefined);
 	}
 
 	/// <summary>
 	/// Plugin for the tray app
 	/// </summary>
-	interface ITrayAppPlugin : IAsyncDisposable
+	interface IToolboxPlugin : IAsyncDisposable
 	{
 		/// <summary>
 		/// Name of the plugin to show in the settings page
@@ -46,7 +46,7 @@ namespace UnrealToolbox
 		/// <summary>
 		/// Get the current status of this plugin
 		/// </summary>
-		TrayAppPluginStatus GetStatus();
+		ToolboxPluginStatus GetStatus();
 
 		/// <summary>
 		/// Determines if the plugin should be shown in the settings page
diff --git a/Engine/Source/Programs/UnrealToolbox/IToolboxPluginHost.cs b/Engine/Source/Programs/UnrealToolbox/IToolboxPluginHost.cs
--- a/Engine/Source/Programs/UnrealToolbox/IToolboxPluginHost.cs
+++ b/Engine/Source/Programs/UnrealToolbox/IToolboxPluginHost.cs
@@ -5,7 +5,7 @@ namespace UnrealToolbox
 	/// <summary>
 	/// Interface for the tray app host application
 	/// </summary>
-	interface ITrayAppHost
+	interface IToolboxPluginHost
 	{
 		/// <summary>
 		/// Notifies the host that a status change has ocurred
diff --git a/Engine/Source/Programs/UnrealToolbox/Plugins/HordeAgent/HordeAgentPlugin.cs b/Engine/Source/Programs/UnrealToolbox/Plugins/HordeAgent/HordeAgentPlugin.cs
--- a/Engine/Source/Programs/UnrealToolbox/Plugins/HordeAgent/HordeAgentPlugin.cs
+++ b/Engine/Source/Programs/UnrealToolbox/Plugins/HordeAgent/HordeAgentPlugin.cs
@@ -15,11 +15,11 @@ using System.Text.Json.Serialization;
 
 namespace UnrealToolbox.Plugins.HordeAgent
 {
-	class HordeAgentPlugin : ITrayAppPlugin
+	class HordeAgentPlugin : IToolboxPlugin
 	{
 		record struct IdleStat(string Name, long Value, long MinValue);
 
-		readonly ITrayAppHost _host;
+		readonly IToolboxPluginHost _host;
 		readonly IToolCatalog _toolCatalog;
 		readonly ILogger _logger;
 
@@ -36,8 +36,8 @@ namespace UnrealToolbox.Plugins.HordeAgent
 
 		AgentSettingsMessage? _agentSettings;
 
-		TrayAppPluginStatus? _status;
-		TrayAppPluginStatus? _reportStatus; // Updated with _status when nothing is currently being installed
+		ToolboxPluginStatus? _status;
+		ToolboxPluginStatus? _reportStatus; // Updated with _status when nothing is currently being installed
 
 		public string Name => "Horde Agent";
 
@@ -75,7 +75,7 @@ namespace UnrealToolbox.Plugins.HordeAgent
 			}
 		}
 
-		public HordeAgentPlugin(ITrayAppHost host, ToolCatalog toolCatalog, ILogger<HordeAgentPlugin> logger)
+		public HordeAgentPlugin(IToolboxPluginHost host, ToolCatalog toolCatalog, ILogger<HordeAgentPlugin> logger)
 		{
 			_host = host;
 			_toolCatalog = toolCatalog;
@@ -233,7 +233,7 @@ namespace UnrealToolbox.Plugins.HordeAgent
 			else if (!status.Healthy)
 			{				
 				string message = String.IsNullOrEmpty(status.Detail) ? "Error. Check logs." : status.Detail.Length > 100 ? status.Detail.Substring(0, 100) : status.Detail;
-				_status = new TrayAppPluginStatus(TrayAppPluginState.Error, message);
+				_status = new ToolboxPluginStatus(TrayAppPluginState.Error, message);
 				
 				if (_settings.Mode != AgentMode.Disabled)
 				{
@@ -254,33 +254,33 @@ namespace UnrealToolbox.Plugins.HordeAgent
 			else if (status.NumLeases > 0)
 			{
 				string message = status.NumLeases == 1 ? "Currently handling 1 lease" : $"Currently handling {status.NumLeases} leases";
-				_status = new TrayAppPluginStatus(TrayAppPluginState.Busy, message);
+				_status = new ToolboxPluginStatus(TrayAppPluginState.Busy, message);
 			}
 			else if (_enabled)
 			{
 				string message = "Agent is operating normally";
-				_status = new TrayAppPluginStatus(TrayAppPluginState.Ok, message);
+				_status = new ToolboxPluginStatus(TrayAppPluginState.Ok, message);
 			}
 			else
 			{
 				string message = "Agent is paused";
-				_status = new TrayAppPluginStatus(TrayAppPluginState.Paused, message);
+				_status = new ToolboxPluginStatus(TrayAppPluginState.Paused, message);
 			}
 			_host.UpdateStatus();
 		}
 
-		public TrayAppPluginStatus GetStatus()
+		public ToolboxPluginStatus GetStatus()
 		{
 			if (_settings.Mode == AgentMode.Disabled)
 			{
-				return TrayAppPluginStatus.Default;
+				return ToolboxPluginStatus.Default;
 			}
 
 			if (!_toolCatalog.Items.Any(x => x.Pending != null && !x.Pending.Failed))
 			{
 				_reportStatus = _status;
 			}
-			return _reportStatus ?? TrayAppPluginStatus.Default;
+			return _reportStatus ?? ToolboxPluginStatus.Default;
 		}
 
 		async Task StatusTaskAsync(CancellationToken cancellationToken)
diff --git a/Engine/Source/Programs/UnrealToolbox/Plugins/HordeProxy/HordeProxyPlugin.cs b/Engine/Source/Programs/UnrealToolbox/Plugins/HordeProxy/HordeProxyPlugin.cs
--- a/Engine/Source/Programs/UnrealToolbox/Plugins/HordeProxy/HordeProxyPlugin.cs
+++ b/Engine/Source/Programs/UnrealToolbox/Plugins/HordeProxy/HordeProxyPlugin.cs
@@ -13,7 +13,7 @@ using Microsoft.Extensions.Hosting;
 
 namespace UnrealToolbox.Plugins.HordeProxy
 {
-	class HordeProxyPlugin : ITrayAppPlugin
+	class HordeProxyPlugin : IToolboxPlugin
 	{
 		readonly BackgroundTask _serverTask;
 		readonly IHordeClientProvider _hordeClientProvider;
@@ -93,8 +93,8 @@ namespace UnrealToolbox.Plugins.HordeProxy
 		{
 		}
 
-		public TrayAppPluginStatus GetStatus()
-			=> TrayAppPluginStatus.Default;
+		public ToolboxPluginStatus GetStatus()
+			=> ToolboxPluginStatus.Default;
 
 		async Task RunServerAsync(CancellationToken cancellationToken)
 		{
diff --git a/Engine/Source/Programs/UnrealToolbox/SettingsWindow.axaml.cs b/Engine/Source/Programs/UnrealToolbox/SettingsWindow.axaml.cs
--- a/Engine/Source/Programs/UnrealToolbox/SettingsWindow.axaml.cs
+++ b/Engine/Source/Programs/UnrealToolbox/SettingsWindow.axaml.cs
@@ -14,8 +14,8 @@ namespace UnrealToolbox
 	partial class SettingsWindow : Window
 	{
 		readonly IServiceProvider? _serviceProvider;
-		readonly Dictionary<string, ITrayAppPlugin> _typeToPlugin = new Dictionary<string, ITrayAppPlugin>();
-		readonly Dictionary<ITrayAppPlugin, Control> _pluginPages = new Dictionary<ITrayAppPlugin, Control>();
+		readonly Dictionary<string, IToolboxPlugin> _typeToPlugin = new Dictionary<string, IToolboxPlugin>();
+		readonly Dictionary<IToolboxPlugin, Control> _pluginPages = new Dictionary<IToolboxPlugin, Control>();
 
 		AboutPage? _aboutPage;
 		GeneralSettingsPage? _generalSettingsPage;
@@ -47,7 +47,7 @@ namespace UnrealToolbox
 				_navView.MenuItems.Clear();
 				_navView.MenuItems.Add(new NavigationViewItem() { Content = "General", IconSource = new SymbolIconSource() { Symbol = Symbol.Settings }, Tag = typeof(GeneralSettingsPage).FullName });
 
-				foreach (ITrayAppPlugin plugin in _serviceProvider.GetServices<ITrayAppPlugin>())
+				foreach (IToolboxPlugin plugin in _serviceProvider.GetServices<IToolboxPlugin>())
 				{
 					if (plugin.HasSettingsPage())
 					{
@@ -74,7 +74,7 @@ namespace UnrealToolbox
 			if (_serviceProvider != null)
 			{
 				bool updated = false;
-				foreach (ITrayAppPlugin plugin in _serviceProvider.GetServices<ITrayAppPlugin>())
+				foreach (IToolboxPlugin plugin in _serviceProvider.GetServices<IToolboxPlugin>())
 				{
 					updated |= plugin.Refresh();
 				}
@@ -100,7 +100,7 @@ namespace UnrealToolbox
 				SettingsContext context = new SettingsContext(this, _serviceProvider!);
 
 				string? typeName = (string?)nvi.Tag;
-				if (typeName != null && _typeToPlugin.TryGetValue(typeName, out ITrayAppPlugin? plugin))
+				if (typeName != null && _typeToPlugin.TryGetValue(typeName, out IToolboxPlugin? plugin))
 				{
 					Control pluginPage = CreatePluginPage(plugin, context);
 					_navView.Content = pluginPage;
@@ -118,7 +118,7 @@ namespace UnrealToolbox
 			}
 		}
 
-		private Control CreatePluginPage(ITrayAppPlugin plugin, SettingsContext context)
+		private Control CreatePluginPage(IToolboxPlugin plugin, SettingsContext context)
 		{
 			Control? pluginPage;
 			if (!_pluginPages.TryGetValue(plugin, out pluginPage))
