Unreal Toolbox: Add a plugin for downloading artifacts.

#rnx

diff --git a/Engine/Source/Programs/UnrealToolbox/App.axaml.cs b/Engine/Source/Programs/UnrealToolbox/App.axaml.cs
--- a/Engine/Source/Programs/UnrealToolbox/App.axaml.cs
+++ b/Engine/Source/Programs/UnrealToolbox/App.axaml.cs
@@ -10,6 +10,7 @@ using Avalonia.Threading;
 using EpicGames.Core;
 using Microsoft.Extensions.DependencyInjection;
 using Microsoft.Extensions.Logging;
+using UnrealToolbox.Plugins.Artifacts;
 using UnrealToolbox.Plugins.HordeAgent;
 using UnrealToolbox.Plugins.HordeProxy;
 
@@ -49,6 +50,7 @@ namespace UnrealToolbox
 			serviceCollection.AddSingleton<ToolboxNotificationManager>();
 
 			// Configure plugins
+			serviceCollection.AddSingleton<IToolboxPlugin, ArtifactsPlugin>();
 			serviceCollection.AddSingleton<IToolboxPlugin, HordeAgentPlugin>();
 			serviceCollection.AddSingleton<IToolboxPlugin, HordeProxyPlugin>();
 
diff --git a/Engine/Source/Programs/UnrealToolbox/Plugins/Artifacts/ArtifactsPlugin.cs b/Engine/Source/Programs/UnrealToolbox/Plugins/Artifacts/ArtifactsPlugin.cs
new file mode 100644
--- /dev/null
+++ b/Engine/Source/Programs/UnrealToolbox/Plugins/Artifacts/ArtifactsPlugin.cs
@@ -0,0 +1,159 @@
+// Copyright Epic Games, Inc. All Rights Reserved.
+
+using Avalonia.Controls;
+using EpicGames.Core;
+using EpicGames.Horde;
+using EpicGames.Horde.Storage;
+using EpicGames.Horde.Storage.Nodes;
+using FluentAvalonia.UI.Controls;
+using Microsoft.Extensions.Logging;
+using Microsoft.Win32;
+
+namespace UnrealToolbox.Plugins.Artifacts
+{
+	class ArtifactsPlugin : IToolboxPlugin
+	{
+		readonly IHordeClientProvider _hordeClientProvider;
+		readonly ILogger _logger;
+
+		public string Name
+			=> "Artifact Downloader";
+
+		public IconSource Icon
+			=> new SymbolIconSource() { Symbol = Symbol.Download };
+
+		public ArtifactsPlugin(IHordeClientProvider hordeClientProvider, ILogger<ArtifactsPlugin> logger)
+		{
+			_hordeClientProvider = hordeClientProvider;
+			_logger = logger;
+		}
+
+		public bool Refresh()
+			=> false;
+
+		public ToolboxPluginStatus GetStatus()
+			=> ToolboxPluginStatus.Default;
+
+		public bool HasSettingsPage()
+			=> false;
+
+		public Control CreateSettingsPage(SettingsContext context)
+			=> throw new NotSupportedException();
+
+		public void PopulateContextMenu(NativeMenu contextMenu)
+		{
+			NativeMenuItem menuItem = new NativeMenuItem("Open artifact...");
+			menuItem.Click += MenuItem_Click;
+			contextMenu.Add(menuItem);
+		}
+
+		private void MenuItem_Click(object? sender, EventArgs e)
+		{
+			_ = DownloadAsync((string?)null, CancellationToken.None);
+		}
+
+		public ValueTask DisposeAsync()
+			=> default;
+
+		public async Task DownloadAsync(string? artifactFile, CancellationToken cancellationToken)
+		{
+			DownloadOptionsWindow window = new DownloadOptionsWindow(artifactFile, _hordeClientProvider);
+			window.Show();
+
+			DownloadOptions? options = await window.Result;
+			if (options != null)
+			{
+				await DownloadAsync(options, cancellationToken);
+			}
+		}
+
+		public async Task DownloadAsync(DownloadOptions options, CancellationToken cancellationToken)
+		{
+			DownloadProgressViewModel progress = new DownloadProgressViewModel();
+
+			using CancellationTokenSource cancellationTokenSource = CancellationTokenSource.CreateLinkedTokenSource(cancellationToken);
+
+			EventHandler? closedEventHandler = (_, _) => cancellationTokenSource.Cancel();
+
+			DownloadProgressWindow progressWindow = new DownloadProgressWindow(progress);
+			progressWindow.Closed += closedEventHandler;
+			progressWindow.Show();
+
+			try
+			{
+				CancellationToken linkedCancellationToken = cancellationTokenSource.Token;
+
+				DirectoryReference stateFolder = DirectoryReference.Combine(options.OutputDir, ".ugs");
+				DirectoryReference.CreateDirectory(stateFolder);
+
+				Uri serverUrl = new Uri(options.BaseUrl.GetLeftPart(UriPartial.Authority));
+
+				using IHordeClientRef? hordeClientRef = _hordeClientProvider.GetClientRef();
+				if (hordeClientRef != null)
+				{
+					IHordeClient hordeClient = hordeClientRef.Client;
+					IStorageNamespace storageNamespace = hordeClient.GetStorageNamespace(options.BaseUrl.AbsolutePath);
+
+					progress.Status = "Connecting to server...";
+
+					IBlobRef<DirectoryNode>? blobRef = await storageNamespace.ReadRefAsync<DirectoryNode>(options.RefName, cancellationToken: linkedCancellationToken);
+
+					DirectoryNode node = await blobRef.ReadBlobAsync(cancellationToken: linkedCancellationToken);
+					progress.SetTotalSize(node.Length);
+
+					progress.Status = "Starting...";
+					await blobRef.ExtractAsync(options.OutputDir.ToDirectoryInfo(), progress, TimeSpan.FromSeconds(0.2), _logger, linkedCancellationToken);
+
+					progress.Status = "Download complete.";
+				}
+
+				progress.OpenFolderPath = options.OutputDir.ToString();
+			}
+			catch (OperationCanceledException)
+			{
+				progressWindow.Close();
+			}
+			catch (Exception ex)
+			{
+				progress.ErrorMessage = $"Error while downloading artifact: {ex.ToString()}";
+			}
+			finally
+			{
+				progress.IsRunning = false;
+				progressWindow.Closed -= closedEventHandler;
+			}
+		}
+
+		public static void RegisterFileAssociations(string application)
+		{
+			string progId = "UnrealGameSync.Artifact";
+
+			using (RegistryKey baseKey = Registry.CurrentUser.CreateSubKey("SOFTWARE\\Classes"))
+			{
+				using (RegistryKey extensionKey = baseKey.CreateSubKey(".uartifact"))
+				{
+					extensionKey.SetValue("", progId);
+				}
+
+				using (RegistryKey progKey = baseKey.CreateSubKey(progId))
+				{
+					progKey.SetValue("", "Horde artifact descriptor");
+
+					using (RegistryKey defaultIconKey = progKey.CreateSubKey("DefaultIcon"))
+					{
+						defaultIconKey.SetValue("", $"\"{application}\",0");
+					}
+
+					using (RegistryKey shellKey = progKey.CreateSubKey("Shell"))
+					{
+						using (RegistryKey openKey = shellKey.CreateSubKey("open"))
+						{
+							using RegistryKey commandKey = openKey.CreateSubKey("command");
+							commandKey.SetValue("", $"\"{application}\" -artifact=\"%1\"");
+						}
+					}
+				}
+			}
+		}
+	}
+}
diff --git a/Engine/Source/Programs/UnrealToolbox/Plugins/Artifacts/DownloadOptions.cs b/Engine/Source/Programs/UnrealToolbox/Plugins/Artifacts/DownloadOptions.cs
new file mode 100644
--- /dev/null
+++ b/Engine/Source/Programs/UnrealToolbox/Plugins/Artifacts/DownloadOptions.cs
@@ -0,0 +1,21 @@
+// Copyright Epic Games, Inc. All Rights Reserved.
+
+using EpicGames.Core;
+using EpicGames.Horde.Storage;
+
+namespace UnrealToolbox.Plugins.Artifacts
+{
+	class DownloadOptions
+	{
+		public Uri BaseUrl { get; set; }
+		public RefName RefName { get; set; }
+		public DirectoryReference OutputDir { get; set; }
+
+		public DownloadOptions(Uri baseUrl, RefName refName, DirectoryReference outputDir)
+		{
+			BaseUrl = baseUrl;
+			RefName = refName;
+			OutputDir = outputDir;
+		}
+	}
+}
diff --git a/Engine/Source/Programs/UnrealToolbox/Plugins/Artifacts/DownloadOptionsViewModel.cs b/Engine/Source/Programs/UnrealToolbox/Plugins/Artifacts/DownloadOptionsViewModel.cs
new file mode 100644
--- /dev/null
+++ b/Engine/Source/Programs/UnrealToolbox/Plugins/Artifacts/DownloadOptionsViewModel.cs
@@ -0,0 +1,262 @@
+// Copyright Epic Games, Inc. All Rights Reserved.
+
+using System.Diagnostics;
+using Avalonia.Controls;
+using Avalonia.Platform.Storage;
+using CommunityToolkit.Mvvm.ComponentModel;
+using EpicGames.Core;
+using EpicGames.Horde;
+using EpicGames.Horde.Artifacts;
+using EpicGames.Horde.Storage;
+using FluentAvalonia.UI.Controls;
+
+namespace UnrealToolbox.Plugins.Artifacts
+{
+	partial class DownloadOptionsViewModel : ObservableObject
+	{
+		readonly Window _window;
+
+		public ArtifactDescriptor? Descriptor { get; private set; }
+
+		public bool ShowInputFile => !String.IsNullOrEmpty(InputFile);
+
+		[ObservableProperty]
+		[NotifyPropertyChangedFor(nameof(AllowStart))]
+		[NotifyPropertyChangedFor(nameof(ShowInputFile))]
+		string? _inputFile;
+
+		[ObservableProperty]
+		string _outputDir = Directory.GetCurrentDirectory();
+
+		[ObservableProperty]
+		[NotifyPropertyChangedFor(nameof(ShowErrorMessage))]
+		string _errorMessage = String.Empty;
+
+		[ObservableProperty]
+		bool _showArtifactInfo = false;
+
+		public bool HasArtifactName => !String.IsNullOrEmpty(ArtifactName);
+
+		[ObservableProperty]
+		[NotifyPropertyChangedFor(nameof(HasArtifactName))]
+		string _artifactName = String.Empty;
+
+		public bool HasArtifactType => !String.IsNullOrEmpty(ArtifactType);
+
+		[ObservableProperty]
+		[NotifyPropertyChangedFor(nameof(HasArtifactType))]
+		string _artifactType = String.Empty;
+
+		public bool HasArtifactDescription => !String.IsNullOrEmpty(ArtifactDescription);
+
+		[ObservableProperty]
+		[NotifyPropertyChangedFor(nameof(HasArtifactDescription))]
+		string _artifactDescription = String.Empty;
+
+		public bool HasArtifactJobUrl => !String.IsNullOrEmpty(ArtifactJobUrl);
+
+		[ObservableProperty]
+		[NotifyPropertyChangedFor(nameof(HasArtifactJobUrl))]
+		string _artifactJobUrl = String.Empty;
+
+		[ObservableProperty]
+		string _artifactKeys = String.Empty;
+
+		[ObservableProperty]
+		string _artifactMetadata = String.Empty;
+
+		public bool ShowErrorMessage
+			=> ErrorMessage.Length > 0;
+
+		public bool AllowStart
+			=> !String.IsNullOrEmpty(InputFile);
+
+		readonly IHordeClientProvider? _hordeClientProvider;
+
+		public DownloadOptionsViewModel(Window window, IHordeClientProvider? hordeClientProvider)
+		{
+			_window = window;
+			_hordeClientProvider = hordeClientProvider;
+		}
+
+		public DownloadOptions? GetDownloadOptions()
+		{
+			if (Descriptor == null)
+			{
+				return null;
+			}
+			if (String.IsNullOrEmpty(OutputDir))
+			{
+				return null;
+			}
+
+			return new DownloadOptions(Descriptor.BaseUrl, Descriptor.RefName, new DirectoryReference(OutputDir));
+		}
+
+		public async Task<bool> SelectUrlAsync()
+		{
+			TextBox urlTextBox = new TextBox();
+			urlTextBox.Text = InputFile;
+			urlTextBox.SelectionStart = InputFile?.Length ?? 0;
+			urlTextBox.SelectionEnd = InputFile?.Length ?? 0;
+
+			ContentDialog dialog = new ContentDialog()
+			{
+				Title = "Open URL",
+				Content = urlTextBox,
+				PrimaryButtonText = "Connect",
+				CloseButtonText = "Cancel"
+			};
+
+			ContentDialogResult result = await dialog.ShowAsync(_window);
+			if (result == ContentDialogResult.None || urlTextBox.Text == null)
+			{
+				return false;
+			}
+
+			try
+			{
+				Uri url = new Uri(urlTextBox.Text);
+				using IHordeClientRef clientRef = _hordeClientProvider!.GetClientRef();
+				using HordeHttpClient httpClient = clientRef.Client.CreateHttpClient();
+
+				string path = url.ToString();
+				if (url.GetLeftPart(UriPartial.Authority) == httpClient.BaseUrl.GetLeftPart(UriPartial.Authority))
+				{
+					path = url.PathAndQuery;
+				}
+
+				GetArtifactResponse response = await HordeHttpRequest.GetAsync<GetArtifactResponse>(httpClient.HttpClient, path, CancellationToken.None);
+
+				ArtifactName = response.Name.ToString();
+				ArtifactType = response.Type.ToString();
+				ArtifactDescription = response.Description ?? String.Empty;
+				ArtifactJobUrl = String.Empty;
+
+				List<string> keys = (response.Keys ?? new List<string>()).Where(x => !x.StartsWith("job:", StringComparison.OrdinalIgnoreCase)).ToList();
+				ArtifactKeys = (keys.Count == 0) ? "(None)" : String.Join(", ", keys);
+				ArtifactMetadata = (response.Metadata == null) ? "(None)" : String.Join(", ", response.Metadata);
+
+				InputFile = url.ToString();
+				ShowArtifactInfo = true;
+				ErrorMessage = String.Empty;
+
+				return true;
+			}
+			catch (Exception ex)
+			{
+				ErrorMessage = ex.Message;
+				return false;
+			}
+		}
+
+		public async Task<bool> SelectInputFileAsync()
+		{
+			FilePickerOpenOptions options = new FilePickerOpenOptions();
+			options.AllowMultiple = false;
+			options.FileTypeFilter = [new FilePickerFileType("Artifacts") { Patterns = ["*.uartifact"] }];
+
+			IReadOnlyList<IStorageFile> files = await _window.StorageProvider.OpenFilePickerAsync(options);
+			if (files.Count == 0)
+			{
+				return false;
+			}
+
+			string? localPath = files[0].TryGetLocalPath();
+			if (localPath == null)
+			{
+				return false;
+			}
+
+			InputFile = localPath;
+
+			FileReference file = new FileReference(InputFile);
+			if (!FileReference.Exists(file))
+			{
+				ErrorMessage = "Artifact descriptor not found.";
+				return false;
+			}
+
+			try
+			{
+				Descriptor = await ArtifactDescriptor.ReadAsync(file, CancellationToken.None);
+
+				ArtifactName = Descriptor.Name ?? String.Empty;
+				ArtifactType = Descriptor.Type ?? String.Empty;
+				ArtifactDescription = Descriptor.Description ?? String.Empty;
+				ArtifactJobUrl = Descriptor.JobUrl?.ToString() ?? String.Empty;
+
+				List<string> keys = (Descriptor.Keys ?? new List<string>()).Where(x => !x.StartsWith("job:", StringComparison.OrdinalIgnoreCase)).ToList();
+				ArtifactKeys = (keys.Count == 0) ? "(None)" : String.Join(", ", keys);
+				ArtifactMetadata = (Descriptor.Metadata == null) ? "(None)" : String.Join(", ", Descriptor.Metadata);
+
+				ShowArtifactInfo = true;
+				ErrorMessage = String.Empty;
+				return true;
+			}
+			catch (Exception ex)
+			{
+				ErrorMessage = ex.ToString();
+				return false;
+			}
+		}
+
+		public void OpenJobUrl()
+		{
+			if (!String.IsNullOrEmpty(ArtifactJobUrl))
+			{
+				string url = ArtifactJobUrl;
+				if (OperatingSystem.IsWindows())
+				{
+					Process.Start(new ProcessStartInfo(url) { UseShellExecute = true })?.Dispose();
+				}
+				else if (OperatingSystem.IsLinux())
+				{
+					Process.Start("xdg-open", url)?.Dispose();
+				}
+				else if (OperatingSystem.IsMacOS())
+				{
+					Process.Start("open", url)?.Dispose();
+				}
+				else
+				{
+					throw new NotImplementedException();
+				}
+			}
+		}
+
+		public void OpenOutputDir()
+		{
+			try
+			{
+				Process.Start(new ProcessStartInfo { FileName = OutputDir, UseShellExecute = true });
+			}
+			catch (Exception ex)
+			{
+				ErrorMessage = ex.Message;
+			}
+		}
+
+		public async Task<bool> SelectOutputDirAsync()
+		{
+			FolderPickerOpenOptions options = new FolderPickerOpenOptions();
+			options.AllowMultiple = false;
+			options.SuggestedStartLocation = await _window.StorageProvider.TryGetFolderFromPathAsync(OutputDir);
+
+			IReadOnlyList<IStorageFolder> folders = await _window.StorageProvider.OpenFolderPickerAsync(options);
+			if (folders.Count == 0)
+			{
+				return false;
+			}
+
+			string? localPath = folders[0].TryGetLocalPath();
+			if (localPath == null)
+			{
+				return false;
+			}
+
+			OutputDir = localPath;
+			return true;
+		}
+	}
+}
diff --git a/Engine/Source/Programs/UnrealToolbox/Plugins/Artifacts/DownloadOptionsWindow.axaml b/Engine/Source/Programs/UnrealToolbox/Plugins/Artifacts/DownloadOptionsWindow.axaml
new file mode 100644
--- /dev/null
+++ b/Engine/Source/Programs/UnrealToolbox/Plugins/Artifacts/DownloadOptionsWindow.axaml
@@ -0,0 +1,80 @@
+<Window xmlns="[REDACTED_PROFILE_URL]"
+        xmlns:x="http://schemas.microsoft.com/winfx/2006/xaml"
+        xmlns:d="http://schemas.microsoft.com/expression/blend/2008"
+        xmlns:mc="http://schemas.openxmlformats.org/markup-compatibility/2006"
+		xmlns:ui="using:FluentAvalonia.UI.Controls"
+		xmlns:local="using:UnrealToolbox.Plugins.Artifacts"
+        mc:Ignorable="d" d:DesignWidth="800" 
+		x:Class="UnrealToolbox.Plugins.Artifacts.DownloadOptionsWindow"
+		x:DataType="local:DownloadOptionsViewModel"
+		Width="800"
+		SizeToContent="Height"
+        Title="Download Artifact"
+		ExtendClientAreaToDecorationsHint="False"
+		WindowStartupLocation="CenterScreen">
+
+	<Window.Styles>
+		<StyleInclude Source="/Styles.axaml"/>
+	</Window.Styles>
+
+	<StackPanel Margin="20 0 20 20" Spacing="5">
+
+		<TextBlock Classes="SettingHeader">Source</TextBlock>
+
+		<Border Classes="Card" CornerRadius="5 5 0 0">
+			<StackPanel Spacing="5" Margin="0">
+				<Grid ColumnDefinitions="*,Auto,Auto" RowDefinitions="Auto">
+					<TextBlock Foreground="Gray" Text="No file selected" IsVisible="{Binding !ShowInputFile}" VerticalAlignment="Center"/>
+					<TextBlock Text="{Binding InputFile}" IsVisible="{Binding ShowInputFile}" VerticalAlignment="Center"/>
+					<Button Content="Browse..." Margin="16 0 0 0" Width="120" Grid.Column="1" Grid.RowSpan="2" Command="{Binding SelectInputFileAsync}"/>
+				</Grid>
+			</StackPanel>
+		</Border>
+		<Border Classes="Card" CornerRadius="0 0 5 5" BorderThickness="1 0 1 1" IsVisible="{Binding ShowArtifactInfo}">
+			<Grid ColumnDefinitions="Auto,*" RowDefinitions="Auto,Auto,Auto,Auto,Auto,Auto">
+				<TextBlock Margin="0 1" Theme="{StaticResource CaptionTextBlockStyle}" Text="Name" Padding="0" Grid.Row="0" IsVisible="{Binding HasArtifactName}"/>
+				<TextBlock Margin="42 1" Theme="{StaticResource CaptionTextBlockStyle}" Text="{Binding ArtifactName}" Padding="0" Grid.Row="0" Grid.Column="1" IsVisible="{Binding HasArtifactName}"/>
+
+				<TextBlock Margin="0 1" Theme="{StaticResource CaptionTextBlockStyle}" Text="Type" Padding="0" Grid.Row="1" IsVisible="{Binding HasArtifactType}"/>
+				<TextBlock Margin="42 1" Theme="{StaticResource CaptionTextBlockStyle}" Text="{Binding ArtifactType}" Padding="0" Grid.Row="1" Grid.Column="1" IsVisible="{Binding HasArtifactType}"/>
+
+				<TextBlock Margin="0 1" Theme="{StaticResource CaptionTextBlockStyle}" Text="Description" Padding="0" Grid.Row="2" IsVisible="{Binding HasArtifactDescription}"/>
+				<TextBlock Margin="42 1" Theme="{StaticResource CaptionTextBlockStyle}" Text="{Binding ArtifactDescription}" Padding="0" Grid.Row="2" Grid.Column="1" IsVisible="{Binding HasArtifactDescription}"/>
+
+				<TextBlock Margin="0 1" Theme="{StaticResource CaptionTextBlockStyle}" Text="Job" Padding="0" Grid.Row="3" IsVisible="{Binding HasArtifactJobUrl}"/>
+				<HyperlinkButton Margin="41 1" FontSize="11" Content="{Binding ArtifactJobUrl}" Padding="0" Grid.Row="3" Grid.Column="1" Command="{Binding OpenJobUrl}" IsVisible="{Binding HasArtifactJobUrl}"/>
+
+				<TextBlock Margin="0 1" Theme="{StaticResource CaptionTextBlockStyle}" Text="Keys" Padding="0" Grid.Row="4"/>
+				<TextBlock Margin="42 1" Theme="{StaticResource CaptionTextBlockStyle}" Text="{Binding ArtifactKeys}" Padding="0" Grid.Row="4" Grid.Column="1"/>
+
+				<TextBlock Margin="0 2" Theme="{StaticResource CaptionTextBlockStyle}" Text="Metadata" Padding="0" Grid.Row="5"/>
+				<TextBlock Margin="42 1" Theme="{StaticResource CaptionTextBlockStyle}" Text="{Binding ArtifactMetadata}" Padding="0" Grid.Column="1" Grid.Row="5"/>
+			</Grid>
+		</Border>
+
+		<TextBlock Classes="SettingHeader">Output</TextBlock>
+
+		<Border Classes="Card" CornerRadius="10 0 10 0">
+			<StackPanel Spacing="5" Margin="0">
+				<Grid ColumnDefinitions="*,Auto" Grid.Column="2" Grid.Row="1">
+					<HyperlinkButton Content="{Binding OutputDir}" Command="{Binding OpenOutputDir}" Margin="0" Padding="0"/>
+					<Button Margin="16 0 0 0" Grid.Column="1" Command="{Binding SelectInputFileAsync}" Width="120">Browse...</Button>
+				</Grid>
+			</StackPanel>
+		</Border>
+
+		<StackPanel Margin="0" Spacing="5" IsVisible="{Binding ShowErrorMessage}">
+			<TextBlock Classes="SettingHeader">Error</TextBlock>
+
+			<Border Classes="Card">
+				<TextBlock Text="{Binding ErrorMessage}" Foreground="Red"/>
+			</Border>
+		</StackPanel>
+
+		<Grid ColumnDefinitions="*,*" HorizontalAlignment="Right">
+			<Button Grid.Column="0" Margin="5 10 10 0" Width="100" Click="Start_Click" IsEnabled="{Binding AllowStart}">Start</Button>
+			<Button Grid.Column="1" Margin="5 10 0 0" Width="100" Click="Cancel_Click">Cancel</Button>
+		</Grid>
+
+	</StackPanel>
+</Window>
diff --git a/Engine/Source/Programs/UnrealToolbox/Plugins/Artifacts/DownloadOptionsWindow.axaml.cs b/Engine/Source/Programs/UnrealToolbox/Plugins/Artifacts/DownloadOptionsWindow.axaml.cs
new file mode 100644
--- /dev/null
+++ b/Engine/Source/Programs/UnrealToolbox/Plugins/Artifacts/DownloadOptionsWindow.axaml.cs
@@ -0,0 +1,63 @@
+// Copyright Epic Games, Inc. All Rights Reserved.
+
+using Avalonia.Controls;
+using Avalonia.Interactivity;
+
+namespace UnrealToolbox.Plugins.Artifacts
+{
+	partial class DownloadOptionsWindow : Window
+	{
+		readonly DownloadOptionsViewModel _viewModel;
+		readonly TaskCompletionSource<DownloadOptions?> _result = new TaskCompletionSource<DownloadOptions?>();
+
+		public Task<DownloadOptions?> Result => _result.Task;
+
+		public DownloadOptionsWindow()
+			: this(null, null)
+		{ }
+
+		public DownloadOptionsWindow(string? artifactFile, IHordeClientProvider? hordeClientProvider)
+		{
+			InitializeComponent();
+
+			_viewModel = new DownloadOptionsViewModel(this, hordeClientProvider);
+			_viewModel.InputFile = artifactFile;
+
+			DataContext = _viewModel;
+		}
+
+		public async Task SelectInputFileAsync()
+		{
+			await _viewModel.SelectInputFileAsync();
+		}
+
+		void Start_Click(object? sender, RoutedEventArgs args)
+		{
+			DownloadOptions? options;
+			try
+			{
+				options = _viewModel.GetDownloadOptions();
+			}
+			catch (Exception ex)
+			{
+				_viewModel.ErrorMessage = $"Unable to read artifact descriptor from {_viewModel.InputFile}: {ex.Message}";
+				return;
+			}
+
+			if (options == null)
+			{
+				_viewModel.ErrorMessage = $"Unable to read artifact descriptor from {_viewModel.InputFile}";
+				return;
+			}
+
+			_result.TrySetResult(options);
+			Close();
+		}
+
+		void Cancel_Click(object? sender, RoutedEventArgs args)
+		{
+			_result.TrySetResult(null);
+			Close();
+		}
+	}
+}
\ No newline at end of file
diff --git a/Engine/Source/Programs/UnrealToolbox/Plugins/Artifacts/DownloadProgressViewModel.cs b/Engine/Source/Programs/UnrealToolbox/Plugins/Artifacts/DownloadProgressViewModel.cs
new file mode 100644
--- /dev/null
+++ b/Engine/Source/Programs/UnrealToolbox/Plugins/Artifacts/DownloadProgressViewModel.cs
@@ -0,0 +1,66 @@
+// Copyright Epic Games, Inc. All Rights Reserved.
+
+using CommunityToolkit.Mvvm.ComponentModel;
+using EpicGames.Horde.Storage.Nodes;
+
+namespace UnrealToolbox.Plugins.Artifacts
+{
+	partial class DownloadProgressViewModel : ObservableObject, IProgress<IExtractStats>
+	{
+		long _size;
+		long _totalSize;
+
+		[ObservableProperty]
+		string _status = "Starting...";
+
+		[ObservableProperty]
+		bool _isRunning = true;
+
+		[ObservableProperty]
+		int _progress;
+
+		[ObservableProperty]
+		bool _isProgressBarVisible;
+
+		[ObservableProperty]
+		[NotifyPropertyChangedFor(nameof(ShowErrorMessage))]
+		[NotifyPropertyChangedFor(nameof(CloseText))]
+		string _errorMessage = String.Empty;
+
+		[ObservableProperty]
+		[NotifyPropertyChangedFor(nameof(ShowOpenFolder))]
+		string? _openFolderPath;
+
+		public bool ShowErrorMessage
+			=> ErrorMessage.Length > 0;
+
+		public bool ShowOpenFolder
+			=> !String.IsNullOrEmpty(OpenFolderPath);
+
+		public string CloseText
+			=> ShowErrorMessage ? "Close" : "Cancel";
+
+		public void Report(IExtractStats value)
+		{
+			Status = $"Copied {value.NumFiles} files ({value.DownloadSize / (1024.0 * 1024.0):n1}mb, {value.DownloadRate / (1024.0 * 1024.0):n1}mb/s)...";
+
+			_size = value.DownloadSize;
+			UpdateProgress();
+		}
+
+		public void SetTotalSize(long totalSize)
+		{
+			_totalSize = totalSize;
+			UpdateProgress();
+		}
+
+		void UpdateProgress()
+		{
+			if (_totalSize != 0)
+			{
+				Progress = (int)((_size * 100) / Math.Max(_totalSize, 1));
+				IsProgressBarVisible = true;
+			}
+		}
+	}
+}
diff --git a/Engine/Source/Programs/UnrealToolbox/Plugins/Artifacts/DownloadProgressWindow.axaml b/Engine/Source/Programs/UnrealToolbox/Plugins/Artifacts/DownloadProgressWindow.axaml
new file mode 100644
--- /dev/null
+++ b/Engine/Source/Programs/UnrealToolbox/Plugins/Artifacts/DownloadProgressWindow.axaml
@@ -0,0 +1,37 @@
+<Window xmlns="[REDACTED_PROFILE_URL]"
+        xmlns:x="http://schemas.microsoft.com/winfx/2006/xaml"
+        xmlns:d="http://schemas.microsoft.com/expression/blend/2008"
+        xmlns:mc="http://schemas.openxmlformats.org/markup-compatibility/2006"
+		xmlns:ui="using:FluentAvalonia.UI.Controls"
+        xmlns:local="using:UnrealToolbox.Plugins.Artifacts"
+        mc:Ignorable="d"
+        x:Class="UnrealToolbox.Plugins.Artifacts.DownloadProgressWindow"
+		x:DataType="local:DownloadProgressViewModel"
+        Title="Downloading Artifact"
+		Width="700"
+		CanResize="False"
+		SizeToContent="Height"
+		WindowStartupLocation="CenterScreen">
+
+	<Window.Styles>
+		<StyleInclude Source="/Styles.axaml"/>
+	</Window.Styles>
+
+	<StackPanel Spacing="20" Margin="32 20">
+		<Grid ColumnDefinitions="*,Auto,Auto" RowDefinitions="Auto,Auto,Auto">
+			<Grid Grid.Column="0" Grid.Row="0" ColumnDefinitions="Auto,*">
+				<ui:ProgressRing Grid.Column="0" Grid.Row="0" Width="16" Height="16" Margin="0 2 6 2" IsVisible="{Binding IsRunning}" VerticalAlignment="Center"/>
+				<TextBlock Grid.Column="1" Grid.Row="0" VerticalAlignment="Center" Theme="{DynamicResource BodyTextBlockStyle}" HorizontalAlignment="Left" Text="{Binding Status}"/>
+			</Grid>
+			<ProgressBar Grid.Column="0" Grid.Row="1" Value="{Binding Progress}" Margin="0 0 12 0"/>
+			<Button Grid.Column="1" Grid.RowSpan="2" Margin="8 0 0 0" VerticalAlignment="Center" Content="Open Folder" IsVisible="{Binding ShowOpenFolder}" Click="OpenFolder"/>
+			<Button Grid.Column="2" Grid.RowSpan="2" Margin="8 0 0 0" VerticalAlignment="Center" Content="{Binding CloseText}" Width="100" Click="Cancel"/>
+		</Grid>
+
+		<Border Classes="Card" IsVisible="{Binding ShowErrorMessage}">
+			<ScrollViewer>
+				<SelectableTextBlock Text="{Binding ErrorMessage}" Foreground="Red" />
+			</ScrollViewer>
+		</Border>
+	</StackPanel>
+</Window>
diff --git a/Engine/Source/Programs/UnrealToolbox/Plugins/Artifacts/DownloadProgressWindow.axaml.cs b/Engine/Source/Programs/UnrealToolbox/Plugins/Artifacts/DownloadProgressWindow.axaml.cs
new file mode 100644
--- /dev/null
+++ b/Engine/Source/Programs/UnrealToolbox/Plugins/Artifacts/DownloadProgressWindow.axaml.cs
@@ -0,0 +1,40 @@
+// Copyright Epic Games, Inc. All Rights Reserved.
+
+using System.Diagnostics;
+using Avalonia.Controls;
+using Avalonia.Interactivity;
+
+namespace UnrealToolbox.Plugins.Artifacts
+{
+	partial class DownloadProgressWindow : Window
+	{
+		readonly DownloadProgressViewModel _viewModel;
+
+		public DownloadProgressWindow()
+			: this(new DownloadProgressViewModel())
+		{
+		}
+
+		public DownloadProgressWindow(DownloadProgressViewModel viewModel)
+		{
+			InitializeComponent();
+
+			_viewModel = viewModel;
+			DataContext = viewModel;
+		}
+
+		public void OpenFolder(object? sender, RoutedEventArgs args)
+		{
+			if (!String.IsNullOrEmpty(_viewModel.OpenFolderPath))
+			{
+				Process.Start(new ProcessStartInfo { FileName = _viewModel.OpenFolderPath, UseShellExecute = true });
+				Close();
+			}
+		}
+
+		public void Cancel(object? sender, RoutedEventArgs args)
+		{
+			Close();
+		}
+	}
+}
diff --git a/Engine/Source/Programs/UnrealToolbox/Styles.axaml b/Engine/Source/Programs/UnrealToolbox/Styles.axaml
--- a/Engine/Source/Programs/UnrealToolbox/Styles.axaml
+++ b/Engine/Source/Programs/UnrealToolbox/Styles.axaml
@@ -7,6 +7,7 @@
 				<StackPanel>
 					<TextBlock Classes="SettingCaption">This is a setting</TextBlock>
 					<TextBlock Classes="SettingHelp">This is what the setting does</TextBlock>
+					<TextBlock Classes="SettingSubHeading">This is what the setting does</TextBlock>
 				</StackPanel>
 			</Border>
 		</Panel>
@@ -34,5 +35,8 @@
 		<Setter Property="FontWeight" Value="Light"/>
 		<Setter Property="VerticalAlignment" Value="Center"/>
 	</Style>
+	<Style Selector="TextBlock.SettingSubHeading">
+		<Setter Property="Theme" Value="{DynamicResource CaptionTextBlockStyle}"/>
+	</Style>
 
 </Styles>
