Unreal Toolbox: Fix race condition in tools catalog initialization

The tools catalog was attempting updates before login completion to the Horde server, leading to failed updates due to missing access token.
This caused a retry cycle followed by a 5-minute delay before the next catalog update attempt.

Fixes:
- Ensure catalog updates only occur after valid token acquisition
- Add manual catalog refresh when server is manually refreshed

diff --git a/Engine/Source/Programs/UnrealToolbox/GeneralSettingsViewModel.cs b/Engine/Source/Programs/UnrealToolbox/GeneralSettingsViewModel.cs
--- a/Engine/Source/Programs/UnrealToolbox/GeneralSettingsViewModel.cs
+++ b/Engine/Source/Programs/UnrealToolbox/GeneralSettingsViewModel.cs
@@ -20,6 +20,7 @@ namespace UnrealToolbox
 
 		readonly SettingsContext _context;
 		readonly IHordeClientProvider _hordeClientProvider;
+		readonly IToolCatalog? _toolCatalog;
 
 		[ObservableProperty]
 		string _serverUrl = String.Empty;
@@ -48,7 +49,8 @@ namespace UnrealToolbox
 			_hordeClientProvider.OnStateChanged += StartRefresh;
 //			_hordeClientProvider.OnAccessTokenStateChanged += StartRefresh;
 
-			ToolCatalog = new ToolCatalogViewModel(context.ServiceProvider.GetService<IToolCatalog>());
+			_toolCatalog = context.ServiceProvider.GetService<IToolCatalog>();
+			ToolCatalog = new ToolCatalogViewModel(_toolCatalog);
 
 			ServerUrl = HordeOptions.GetDefaultServerUrl()?.ToString() ?? "Unknown";
 			ServerStatus = "Unknown";
@@ -168,6 +170,7 @@ namespace UnrealToolbox
 				}
 				else
 				{
+					_toolCatalog?.RequestUpdate();
 					return new ConnectionState(client.ServerUrl.ToString(), "Authenticated");
 				}
 			}
diff --git a/Engine/Source/Programs/UnrealToolbox/IToolCatalog.cs b/Engine/Source/Programs/UnrealToolbox/IToolCatalog.cs
--- a/Engine/Source/Programs/UnrealToolbox/IToolCatalog.cs
+++ b/Engine/Source/Programs/UnrealToolbox/IToolCatalog.cs
@@ -11,6 +11,12 @@ namespace UnrealToolbox
 		/// Whether to auto-update tools to latest
 		/// </summary>
 		bool AutoUpdate { get; set; }
+		
+		/// <summary>
+		/// Asynchronously request an update of tool catalog
+		/// Listen on <see cref="OnItemsChanged"/> for when items actually update.
+		/// </summary>
+		void RequestUpdate();
 
 		/// <summary>
 		/// List of available tools
diff --git a/Engine/Source/Programs/UnrealToolbox/ToolCatalog.cs b/Engine/Source/Programs/UnrealToolbox/ToolCatalog.cs
--- a/Engine/Source/Programs/UnrealToolbox/ToolCatalog.cs
+++ b/Engine/Source/Programs/UnrealToolbox/ToolCatalog.cs
@@ -2,6 +2,7 @@
 
 using System.Diagnostics;
 using System.Runtime.InteropServices;
+using System.Security.Authentication;
 using System.Text.Json;
 using EpicGames.Core;
 using EpicGames.Horde;
@@ -228,6 +229,12 @@ namespace UnrealToolbox
 				}
 			}
 		}
+		
+		/// <inheritdoc/>
+		public void RequestUpdate()
+		{
+			_updateEvent.Pulse();
+		}
 
 		public ToolCatalog(IHordeClientProvider hordeClientProvider, ILogger<ToolCatalog> logger)
 		{
@@ -938,6 +945,12 @@ namespace UnrealToolbox
 					catch (OperationCanceledException)
 					{
 					}
+					catch (InvalidCredentialException)
+					{
+						// Briefly wait and try again as missing access token check is cheap and in-memory.
+						await Task.WhenAny(updateTask, Task.Delay(TimeSpan.FromSeconds(5), cancellationToken));
+						continue;
+					}
 					catch (Exception ex)
 					{
 						_logger.LogError(ex, "Exception while checking for tool updates: {Message}", ex.Message);
@@ -962,6 +975,11 @@ namespace UnrealToolbox
 				SetTools(Enumerable.Empty<ITool>());
 				_serverUri = hordeClientRef.Client.ServerUrl;
 			}
+			
+			if (!hordeClientRef.Client.HasValidAccessToken())
+			{
+				throw new InvalidCredentialException("No access token available. Auth via Toolbox UI is required.");
+			}
 
 			IEnumerable<ITool> tools = await hordeClientRef.Client.Tools.GetAllAsync(cancellationToken);
 			SetTools(tools);
