Unreal Toolbox: Add support to hide process window for install commands

diff --git a/Engine/Source/Programs/UnrealToolbox/ToolCatalog.cs b/Engine/Source/Programs/UnrealToolbox/ToolCatalog.cs
--- a/Engine/Source/Programs/UnrealToolbox/ToolCatalog.cs
+++ b/Engine/Source/Programs/UnrealToolbox/ToolCatalog.cs
@@ -514,7 +514,8 @@ namespace UnrealToolbox
 				PendingToolDeploymentInfo pending = new PendingToolDeploymentInfo(false, "Running install actions...", deploymentInfo);
 				UpdateStateAndNotify(item, state => state with { Pending = pending });
 
-				int exitCode = await RunCommandAsync(item.Id.ToString(), toolConfig.InstallCommand, toolDir.FullName, toolConfig.RequiresElevation, cancellationToken);
+				int exitCode = await RunCommandAsync(item.Id.ToString(), toolConfig.InstallCommand, toolDir.FullName, 
+					toolConfig.RequiresElevation, toolConfig.Hidden, cancellationToken);
 				if (exitCode != 0)
 				{
 					pending = new PendingToolDeploymentInfo(true, $"Installation failed ({exitCode})", deploymentInfo, ShowLogLink: true);
@@ -550,7 +551,8 @@ namespace UnrealToolbox
 					ToolCommand installCommand = new ToolCommand { FileName = "msiexec.exe" };
 					installCommand.Arguments = new List<string> { "/I", msiFile.FullName };
 
-					int exitCode = await RunCommandAsync(item.Id.ToString(), installCommand, toolDir.FullName, toolConfig.RequiresElevation, cancellationToken);
+					int exitCode = await RunCommandAsync(item.Id.ToString(), installCommand, toolDir.FullName, 
+						toolConfig.RequiresElevation, toolConfig.Hidden, cancellationToken);
 					if (exitCode != 0 || !IsMsiInstalled(item._state.MsiProductId))
 					{
 						UpdateStateAndNotify(item, state => state with { Pending = null });
@@ -597,7 +599,8 @@ namespace UnrealToolbox
 					PendingToolDeploymentInfo? pending = new PendingToolDeploymentInfo(false, $"Removing {current.Version}", null);
 					UpdateStateAndNotify(item, state => state with { Pending = pending, Current = null });
 
-					await RunCommandAsync(item.Id.ToString(), current.Config.UninstallCommand, current.Dir.FullName, current.Config.RequiresElevation, cancellationToken);
+					await RunCommandAsync(item.Id.ToString(), current.Config.UninstallCommand, current.Dir.FullName, 
+						current.Config.RequiresElevation, current.Config.Hidden, cancellationToken);
 				}
 				else if (OperatingSystem.IsWindows() && !String.IsNullOrEmpty(item._state.MsiProductId) && IsMsiInstalled(item._state.MsiProductId))
 				{
@@ -621,7 +624,8 @@ namespace UnrealToolbox
 						ToolCommand uninstallCommand = new ToolCommand { FileName = "msiexec.exe" };
 						uninstallCommand.Arguments = new List<string> { "/x", $"{{{item._state.MsiProductId}}}" };
 
-						int exitCode = await RunCommandAsync(item.Id.ToString(), uninstallCommand, Directory.GetCurrentDirectory(), current.Config.RequiresElevation, cancellationToken);
+						int exitCode = await RunCommandAsync(item.Id.ToString(), uninstallCommand, Directory.GetCurrentDirectory(), 
+							current.Config.RequiresElevation, current.Config.Hidden, cancellationToken);
 						if (exitCode == 1602)
 						{
 							// User cancelled (https://learn.microsoft.com/en-us/windows/win32/msi/error-codes)
@@ -644,7 +648,7 @@ namespace UnrealToolbox
 			return true;
 		}
 
-		async Task<int> RunCommandAsync(string toolName, ToolCommand command, string workingDir, bool elevateProcess, CancellationToken cancellationToken)
+		async Task<int> RunCommandAsync(string toolName, ToolCommand command, string workingDir, bool elevateProcess, bool hidden, CancellationToken cancellationToken)
 		{
 			string arguments = String.Empty;
 			if (command.Arguments != null)
@@ -666,8 +670,14 @@ namespace UnrealToolbox
 				FileName = fileName,
 				Arguments = arguments,
 				WorkingDirectory = workingDir,
+				CreateNoWindow = true
 			};
 			
+			if (hidden)
+			{
+				process.StartInfo.WindowStyle = ProcessWindowStyle.Hidden;
+			}
+			
 			if (elevateProcess)
 			{
 				// When elevating the process, IO cannot be redirected in this process
@@ -680,7 +690,6 @@ namespace UnrealToolbox
 				process.StartInfo.RedirectStandardOutput = true;
 				process.StartInfo.RedirectStandardError = true;
 				process.StartInfo.UseShellExecute = false;
-				process.StartInfo.CreateNoWindow = true;
 				
 				DataReceivedEventHandler onDataReceived = (sender, args) =>
 				{
diff --git a/Engine/Source/Programs/UnrealToolbox/ToolConfig.cs b/Engine/Source/Programs/UnrealToolbox/ToolConfig.cs
--- a/Engine/Source/Programs/UnrealToolbox/ToolConfig.cs
+++ b/Engine/Source/Programs/UnrealToolbox/ToolConfig.cs
@@ -21,6 +21,11 @@ namespace UnrealToolbox
 		/// Whether the install/uninstall for this tool needs administrator or root privileges
 		/// </summary>
 		public bool RequiresElevation { get; set; } = false;
+		
+		/// <summary>
+		/// Whether the install/uninstall process window should be hidden
+		/// </summary>
+		public bool Hidden { get; set; } = false;
 
 		/// <summary>
 		/// Command to run when installing
