Unreal Toolbox: Fix exception disposing cancellation token source when another exception is thrown during tool catalog update.

#rnx

diff --git a/Engine/Source/Programs/UnrealToolbox/ToolCatalog.cs b/Engine/Source/Programs/UnrealToolbox/ToolCatalog.cs
--- a/Engine/Source/Programs/UnrealToolbox/ToolCatalog.cs
+++ b/Engine/Source/Programs/UnrealToolbox/ToolCatalog.cs
@@ -958,16 +958,18 @@ namespace UnrealToolbox
 					{
 						// Briefly wait and try again as missing access token check is cheap and in-memory.
 						await Task.WhenAny(updateTask, Task.Delay(TimeSpan.FromSeconds(5), cancellationToken));
-						continue;
+						continue; 
 					}
 					catch (Exception ex)
 					{
 						_logger.LogError(ex, "Exception while checking for tool updates: {Message}", ex.Message);
 					}
-
-					lock (cancellationSource)
+					finally
 					{
-						complete = true;
+						lock (cancellationSource)
+						{
+							complete = true;
+						}
 					}
 				}
 
