Unreal Toolbox: Add support for the horde-artifact:// URL handler, allowing toolbox to open links from Horde for download.

diff --git a/Engine/Source/Programs/UnrealToolbox/App.axaml.cs b/Engine/Source/Programs/UnrealToolbox/App.axaml.cs
--- a/Engine/Source/Programs/UnrealToolbox/App.axaml.cs
+++ b/Engine/Source/Programs/UnrealToolbox/App.axaml.cs
@@ -1,6 +1,7 @@
 // Copyright Epic Games, Inc. All Rights Reserved.
 
 using System.Diagnostics;
+using System.Reflection;
 using System.Runtime.InteropServices;
 using Avalonia;
 using Avalonia.Controls;
@@ -10,6 +11,7 @@ using Avalonia.Threading;
 using EpicGames.Core;
 using Microsoft.Extensions.DependencyInjection;
 using Microsoft.Extensions.Logging;
+using Microsoft.Win32;
 using UnrealToolbox.Plugins.Artifacts;
 using UnrealToolbox.Plugins.HordeAgent;
 using UnrealToolbox.Plugins.HordeProxy;
@@ -98,6 +100,40 @@ namespace UnrealToolbox
 			_settingsThreadStop = new ManualResetEvent(false);
 			_settingsThread = new Thread(WaitForEvents);
 			_settingsThread.Start();
+
+			RegisterProtocolHandlers();
+		}
+
+		void RegisterProtocolHandlers()
+		{
+			List<string> protocols = new List<string>();
+			foreach (IToolboxPlugin plugin in _serviceProvider.GetServices<IToolboxPlugin>())
+			{
+				protocols.AddRange(plugin.UrlProtocols);
+			}
+
+			if (OperatingSystem.IsWindows())
+			{
+				string exeFile = Path.ChangeExtension(Assembly.GetExecutingAssembly().Location, ".exe");
+				foreach (string protocol in protocols)
+				{
+					using (RegistryKey baseKey = Registry.CurrentUser.CreateSubKey("SOFTWARE\\Classes\\horde-artifact"))
+					{
+						baseKey.SetValue(null, $"URL:{protocol} protocol");
+						baseKey.SetValue("URL Protocol", "");
+
+						using (RegistryKey iconKey = baseKey.CreateSubKey("DefaultIcon"))
+						{
+							iconKey.SetValue(null, $"\"{exeFile}\",0");
+						}
+
+						using (RegistryKey commandKey = baseKey.CreateSubKey("shell\\open\\command"))
+						{
+							commandKey.SetValue(null, $"\"{exeFile}\" -url=\"%1\" -quiet");
+						}
+					}
+				}
+			}
 		}
 
 		private void WaitForEvents()
@@ -105,9 +141,10 @@ namespace UnrealToolbox
 			using EventWaitHandle closeEvent = new EventWaitHandle(false, EventResetMode.AutoReset, Program.CloseEventName);
 			using EventWaitHandle refreshEvent = new EventWaitHandle(false, EventResetMode.AutoReset, Program.RefreshEventName);
 			using EventWaitHandle settingsEvent = new EventWaitHandle(false, EventResetMode.AutoReset, Program.SettingsEventName);
+			using EventWaitHandle commandEvent = new EventWaitHandle(false, EventResetMode.AutoReset, Program.CommandEventName);
 			for (; ; )
 			{
-				int index = WaitHandle.WaitAny(new[] { closeEvent, refreshEvent, settingsEvent, _settingsThreadStop! });
+				int index = WaitHandle.WaitAny(new[] { closeEvent, refreshEvent, settingsEvent, commandEvent, _settingsThreadStop! });
 				if (index == 0)
 				{
 					Dispatcher.UIThread.Post(() => CloseMainThread());
@@ -120,6 +157,10 @@ namespace UnrealToolbox
 				{
 					Dispatcher.UIThread.Post(() => OpenSettings());
 				}
+				else if (index == 3)
+				{
+					Dispatcher.UIThread.Post(() => ExecuteCommands());
+				}
 				else
 				{
 					break;
@@ -132,6 +173,37 @@ namespace UnrealToolbox
 			((IClassicDesktopStyleApplicationLifetime)ApplicationLifetime!).Shutdown();
 		}
 
+		void ExecuteCommands()
+		{
+			List<Command> commands = CommandHelper.DequeueAll();
+			foreach (Command command in commands)
+			{
+				try
+				{
+					ExecuteCommand(command);
+				}
+				catch (Exception ex)
+				{
+					_logger.LogError(ex, "Error executing command: {Message}", ex.Message);
+				}
+			}
+		}
+
+		void ExecuteCommand(Command command)
+		{
+			if (command.Type == CommandType.OpenUrl)
+			{
+				Uri url = new Uri(command.Argument!);
+				foreach (IToolboxPlugin plugin in _serviceProvider.GetServices<IToolboxPlugin>())
+				{
+					if (plugin.HandleUrl(url))
+					{
+						break;
+					}
+				}
+			}
+		}
+
 		private void UpdateReady()
 		{
 			Dispatcher.UIThread.Post(() => UpdateReadyMainThread());
diff --git a/Engine/Source/Programs/UnrealToolbox/Command.cs b/Engine/Source/Programs/UnrealToolbox/Command.cs
new file mode 100644
--- /dev/null
+++ b/Engine/Source/Programs/UnrealToolbox/Command.cs
@@ -0,0 +1,117 @@
+// Copyright Epic Games, Inc. All Rights Reserved.
+
+using System.Text.Json;
+using System.Text.Json.Serialization;
+using EpicGames.Core;
+
+namespace UnrealToolbox
+{
+	[JsonConverter(typeof(JsonStringEnumConverter))]
+	enum CommandType
+	{
+		OpenUrl,
+	}
+
+	class Command
+	{
+		public CommandType Type { get; set; }
+		public string? Argument { get; set; }
+	}
+
+	class CommandList
+	{
+		public List<Command> Commands { get; set; } = new List<Command>();
+	}
+
+	class CommandHelper
+	{
+		const string MutexName = "UnrealToolbox-CommandMutex";
+
+		static readonly FileReference s_commandFile = FileReference.Combine(Program.DataDir, "Commands.json");
+
+		static readonly JsonSerializerOptions s_jsonOptions = new JsonSerializerOptions
+		{
+			AllowTrailingCommas = true,
+			DefaultIgnoreCondition = JsonIgnoreCondition.WhenWritingNull,
+			PropertyNameCaseInsensitive = true,
+			PropertyNamingPolicy = JsonNamingPolicy.CamelCase,
+			WriteIndented = true
+		};
+
+		public static void Enqueue(Command command)
+		{
+			RunUnderMutex(() =>
+			{
+				CommandList commandList = ReadCommandList();
+				commandList.Commands.Add(command);
+				WriteCommandList(commandList);
+			});
+		}
+
+		public static List<Command> DequeueAll()
+		{
+			CommandList? commandList = null;
+			RunUnderMutex(() =>
+			{
+				commandList = ReadCommandList();
+				FileReference.Delete(s_commandFile);
+			});
+			return commandList?.Commands ?? new List<Command>();
+		}
+
+		static void RunUnderMutex(Action action)
+		{
+			using Mutex mutex = new Mutex(false, MutexName);
+			try
+			{
+				mutex.WaitOne();
+			}
+			catch (AbandonedMutexException)
+			{
+				// Ignore
+			}
+
+			try
+			{
+				action();
+			}
+			finally
+			{
+				mutex.ReleaseMutex();
+			}
+		}
+
+		static CommandList ReadCommandList()
+		{
+			CommandList? commandList = null;
+			using (FileStream? stream = FileTransaction.OpenRead(s_commandFile))
+			{
+				if (stream == null)
+				{
+					commandList = new CommandList();
+				}
+				else
+				{
+					try
+					{
+						commandList = JsonSerializer.Deserialize<CommandList>(stream, s_jsonOptions) ?? new CommandList();
+					}
+					catch
+					{
+						commandList = new CommandList();
+					}
+				}
+			}
+			return commandList;
+		}
+
+		static void WriteCommandList(CommandList commandList)
+		{
+			using (FileTransactionStream? stream = FileTransaction.OpenWrite(s_commandFile))
+			{
+				JsonSerializer.Serialize(stream, commandList, s_jsonOptions);
+				stream.CompleteTransaction();
+			}
+		}
+	}
+}
diff --git a/Engine/Source/Programs/UnrealToolbox/IToolboxPlugin.cs b/Engine/Source/Programs/UnrealToolbox/IToolboxPlugin.cs
--- a/Engine/Source/Programs/UnrealToolbox/IToolboxPlugin.cs
+++ b/Engine/Source/Programs/UnrealToolbox/IToolboxPlugin.cs
@@ -37,12 +37,22 @@ namespace UnrealToolbox
 		/// </summary>
 		IconSource Icon { get; }
 
+		/// <summary>
+		/// Url protocol handlers that this plugin supports
+		/// </summary>
+		IReadOnlyList<string> UrlProtocols { get; }
+
 		/// <summary>
 		/// Refresh the state of this plugin. Called when the application is activated or focussed.
 		/// </summary>
 		/// <returns>True if the plugin status has changed, and the context menu needs to be rebuilt</returns>
 		bool Refresh();
 
+		/// <summary>
+		/// Handles a url passed on the commandline
+		/// </summary>
+		bool HandleUrl(Uri url);
+
 		/// <summary>
 		/// Get the current status of this plugin
 		/// </summary>
@@ -63,4 +73,52 @@ namespace UnrealToolbox
 		/// </summary>
 		void PopulateContextMenu(NativeMenu contextMenu);
 	}
+
+	/// <summary>
+	/// Base implementation of <see cref="IToolboxPlugin"/>
+	/// </summary>
+	abstract class ToolboxPluginBase : IToolboxPlugin
+	{
+		/// <inheritdoc/>
+		public abstract string Name { get; }
+
+		/// <inheritdoc/>
+		public virtual IconSource Icon
+			=> new SymbolIconSource() { Symbol = Symbol.Help };
+
+		/// <inheritdoc/>
+		public virtual IReadOnlyList<string> UrlProtocols
+			=> Array.Empty<string>();
+
+		/// <inheritdoc/>
+		public virtual bool Refresh()
+			=> false;
+
+		/// <inheritdoc/>
+		public virtual bool HandleUrl(Uri url)
+			=> false;
+
+		/// <inheritdoc/>
+		public virtual Control CreateSettingsPage(SettingsContext context)
+			=> throw new NotSupportedException();
+
+		/// <inheritdoc/>
+		public virtual ValueTask DisposeAsync()
+		{
+			GC.SuppressFinalize(this);
+			return default;
+		}
+
+		/// <inheritdoc/>
+		public virtual ToolboxPluginStatus GetStatus()
+			=> ToolboxPluginStatus.Default;
+
+		/// <inheritdoc/>
+		public virtual bool HasSettingsPage()
+			=> false;
+
+		/// <inheritdoc/>
+		public virtual void PopulateContextMenu(NativeMenu contextMenu)
+		{ }
+	}
 }
diff --git a/Engine/Source/Programs/UnrealToolbox/Plugins/Artifacts/ArtifactsPlugin.cs b/Engine/Source/Programs/UnrealToolbox/Plugins/Artifacts/ArtifactsPlugin.cs
--- a/Engine/Source/Programs/UnrealToolbox/Plugins/Artifacts/ArtifactsPlugin.cs
+++ b/Engine/Source/Programs/UnrealToolbox/Plugins/Artifacts/ArtifactsPlugin.cs
@@ -11,36 +11,38 @@ using Microsoft.Win32;
 
 namespace UnrealToolbox.Plugins.Artifacts
 {
-	class ArtifactsPlugin : IToolboxPlugin
+	class ArtifactsPlugin : ToolboxPluginBase
 	{
 		readonly IHordeClientProvider _hordeClientProvider;
 		readonly ILogger _logger;
 
-		public string Name
+		public override string Name
 			=> "Artifact Downloader";
 
-		public IconSource Icon
+		public override IconSource Icon
 			=> new SymbolIconSource() { Symbol = Symbol.Download };
 
+		public override IReadOnlyList<string> UrlProtocols
+			=> new[] { "horde-artifact" };
+
 		public ArtifactsPlugin(IHordeClientProvider hordeClientProvider, ILogger<ArtifactsPlugin> logger)
 		{
 			_hordeClientProvider = hordeClientProvider;
 			_logger = logger;
 		}
 
-		public bool Refresh()
-			=> false;
-
-		public ToolboxPluginStatus GetStatus()
-			=> ToolboxPluginStatus.Default;
-
-		public bool HasSettingsPage()
-			=> false;
+		public override bool HandleUrl(Uri url)
+		{
+			if (!url.Scheme.Equals("horde-artifact", StringComparison.OrdinalIgnoreCase))
+			{
+				return false;
+			}
 
-		public Control CreateSettingsPage(SettingsContext context)
-			=> throw new NotSupportedException();
+			_ = DownloadAsync(url, CancellationToken.None);
+			return true;
+		}
 
-		public void PopulateContextMenu(NativeMenu contextMenu)
+		public override void PopulateContextMenu(NativeMenu contextMenu)
 		{
 			NativeMenuItem menuItem = new NativeMenuItem("Open artifact...");
 			menuItem.Click += MenuItem_Click;
@@ -49,25 +51,22 @@ namespace UnrealToolbox.Plugins.Artifacts
 
 		private void MenuItem_Click(object? sender, EventArgs e)
 		{
-			_ = DownloadAsync((string?)null, CancellationToken.None);
+			_ = DownloadAsync(null, CancellationToken.None);
 		}
 
-		public ValueTask DisposeAsync()
-			=> default;
-
-		public async Task DownloadAsync(string? artifactFile, CancellationToken cancellationToken)
+		public async Task DownloadAsync(Uri? artifactUrl, CancellationToken cancellationToken)
 		{
-			DownloadOptionsWindow window = new DownloadOptionsWindow(artifactFile, _hordeClientProvider);
+			DownloadOptionsWindow window = new DownloadOptionsWindow(artifactUrl, _hordeClientProvider);
 			window.Show();
 
 			DownloadOptions? options = await window.Result;
 			if (options != null)
 			{
-				await DownloadAsync(options, cancellationToken);
+				await DownloadDataAsync(options, cancellationToken);
 			}
 		}
 
-		public async Task DownloadAsync(DownloadOptions options, CancellationToken cancellationToken)
+		public async Task DownloadDataAsync(DownloadOptions options, CancellationToken cancellationToken)
 		{
 			DownloadProgressViewModel progress = new DownloadProgressViewModel();
 
diff --git a/Engine/Source/Programs/UnrealToolbox/Plugins/Artifacts/DownloadOptionsViewModel.cs b/Engine/Source/Programs/UnrealToolbox/Plugins/Artifacts/DownloadOptionsViewModel.cs
--- a/Engine/Source/Programs/UnrealToolbox/Plugins/Artifacts/DownloadOptionsViewModel.cs
+++ b/Engine/Source/Programs/UnrealToolbox/Plugins/Artifacts/DownloadOptionsViewModel.cs
@@ -12,18 +12,55 @@ using FluentAvalonia.UI.Controls;
 
 namespace UnrealToolbox.Plugins.Artifacts
 {
+	record class ArtifactInfo(string? Name, string? Type, string? Description, Uri? JobUrl, IReadOnlyList<string>? Keys, IReadOnlyList<string>? Metadata, Uri BaseUrl, RefName RefName);
+
 	partial class DownloadOptionsViewModel : ObservableObject
 	{
 		readonly Window _window;
 
-		public ArtifactDescriptor? Descriptor { get; private set; }
+		[ObservableProperty]
+		[NotifyPropertyChangedFor(nameof(InputText))]
+		Uri? _artifactUrl;
 
-		public bool ShowInputFile => !String.IsNullOrEmpty(InputFile);
+		[ObservableProperty]
+		bool _showUrlSpinner;
 
 		[ObservableProperty]
 		[NotifyPropertyChangedFor(nameof(AllowStart))]
-		[NotifyPropertyChangedFor(nameof(ShowInputFile))]
-		string? _inputFile;
+		[NotifyPropertyChangedFor(nameof(ShowArtifactInfo))]
+		[NotifyPropertyChangedFor(nameof(HasArtifactName))]
+		[NotifyPropertyChangedFor(nameof(ArtifactName))]
+		[NotifyPropertyChangedFor(nameof(HasArtifactType))]
+		[NotifyPropertyChangedFor(nameof(ArtifactType))]
+		[NotifyPropertyChangedFor(nameof(HasArtifactDescription))]
+		[NotifyPropertyChangedFor(nameof(ArtifactDescription))]
+		[NotifyPropertyChangedFor(nameof(HasArtifactJobUrl))]
+		[NotifyPropertyChangedFor(nameof(ArtifactJobUrl))]
+		[NotifyPropertyChangedFor(nameof(ArtifactKeys))]
+		[NotifyPropertyChangedFor(nameof(ArtifactMetadata))]
+		ArtifactInfo? _artifact;
+
+		public string InputText
+		{
+			get
+			{
+				if (ArtifactUrl == null)
+				{
+					return "No file selected.";
+				}
+				else if (ArtifactUrl.IsFile)
+				{
+					return ArtifactUrl.LocalPath;
+				}
+				else
+				{
+					return ArtifactUrl.ToString();
+				}
+			}
+		}
+
+		public bool ShowBrowse
+			=> ArtifactUrl == null || ArtifactUrl.IsFile;
 
 		[ObservableProperty]
 		string _outputDir = Directory.GetCurrentDirectory();
@@ -32,44 +69,48 @@ namespace UnrealToolbox.Plugins.Artifacts
 		[NotifyPropertyChangedFor(nameof(ShowErrorMessage))]
 		string _errorMessage = String.Empty;
 
-		[ObservableProperty]
-		bool _showArtifactInfo = false;
+		public bool ShowArtifactInfo
+			=> Artifact != null;
 
 		public bool HasArtifactName => !String.IsNullOrEmpty(ArtifactName);
 
-		[ObservableProperty]
-		[NotifyPropertyChangedFor(nameof(HasArtifactName))]
-		string _artifactName = String.Empty;
+		public string ArtifactName
+			=> Artifact?.Name ?? String.Empty;
 
-		public bool HasArtifactType => !String.IsNullOrEmpty(ArtifactType);
+		public bool HasArtifactType => !String.IsNullOrEmpty(Artifact?.Type);
 
-		[ObservableProperty]
-		[NotifyPropertyChangedFor(nameof(HasArtifactType))]
-		string _artifactType = String.Empty;
+		public string ArtifactType
+			=> Artifact?.Type ?? String.Empty;
 
-		public bool HasArtifactDescription => !String.IsNullOrEmpty(ArtifactDescription);
+		public bool HasArtifactDescription
+			=> !String.IsNullOrEmpty(ArtifactDescription);
 
-		[ObservableProperty]
-		[NotifyPropertyChangedFor(nameof(HasArtifactDescription))]
-		string _artifactDescription = String.Empty;
+		public string ArtifactDescription
+			=> Artifact?.Description ?? String.Empty;
 
-		public bool HasArtifactJobUrl => !String.IsNullOrEmpty(ArtifactJobUrl);
+		public bool HasArtifactJobUrl
+			=> !String.IsNullOrEmpty(ArtifactJobUrl);
 
-		[ObservableProperty]
-		[NotifyPropertyChangedFor(nameof(HasArtifactJobUrl))]
-		string _artifactJobUrl = String.Empty;
+		public string ArtifactJobUrl
+			=> Artifact?.JobUrl?.ToString() ?? String.Empty;
 
-		[ObservableProperty]
-		string _artifactKeys = String.Empty;
+		public string ArtifactKeys
+		{
+			get
+			{
+				List<string> keysList = (Artifact?.Keys ?? new List<string>()).Where(x => !x.StartsWith("job:", StringComparison.OrdinalIgnoreCase)).ToList();
+				return (keysList.Count == 0) ? "(None)" : String.Join(", ", keysList);
+			}
+		}
 
-		[ObservableProperty]
-		string _artifactMetadata = String.Empty;
+		public string ArtifactMetadata
+			=> (Artifact?.Metadata == null) ? "(None)" : String.Join(", ", Artifact.Metadata);
 
 		public bool ShowErrorMessage
 			=> ErrorMessage.Length > 0;
 
 		public bool AllowStart
-			=> !String.IsNullOrEmpty(InputFile);
+			=> Artifact != null;
 
 		readonly IHordeClientProvider? _hordeClientProvider;
 
@@ -81,7 +122,7 @@ namespace UnrealToolbox.Plugins.Artifacts
 
 		public DownloadOptions? GetDownloadOptions()
 		{
-			if (Descriptor == null)
+			if (Artifact == null)
 			{
 				return null;
 			}
@@ -90,15 +131,76 @@ namespace UnrealToolbox.Plugins.Artifacts
 				return null;
 			}
 
-			return new DownloadOptions(Descriptor.BaseUrl, Descriptor.RefName, new DirectoryReference(OutputDir));
+			return new DownloadOptions(Artifact.BaseUrl, Artifact.RefName, new DirectoryReference(OutputDir));
+		}
+
+		public async Task<bool> RefreshArtifactAsync()
+		{
+			ShowUrlSpinner = true;
+			try
+			{
+				Uri? artifactUrl = ArtifactUrl;
+				if (artifactUrl == null)
+				{
+					Artifact = null;
+
+					ErrorMessage = String.Empty;
+					return false;
+				}
+				else if (artifactUrl.IsFile)
+				{
+					FileReference file = new FileReference(artifactUrl.LocalPath);
+					if (!FileReference.Exists(file))
+					{
+						ErrorMessage = "Artifact descriptor not found.";
+						return false;
+					}
+
+					ArtifactDescriptor descriptor = await ArtifactDescriptor.ReadAsync(file, CancellationToken.None);
+					Artifact = new ArtifactInfo(descriptor.Name, descriptor.Type, descriptor.Description, descriptor.JobUrl, descriptor.Keys, descriptor.Metadata, descriptor.BaseUrl, descriptor.RefName);
+
+					ErrorMessage = String.Empty;
+					return true;
+				}
+				else
+				{
+					using IHordeClientRef clientRef = _hordeClientProvider!.GetClientRef();
+					using HordeHttpClient httpClient = clientRef.Client.CreateHttpClient();
+
+					if (artifactUrl.Host != httpClient.BaseUrl.Host)
+					{
+						ErrorMessage = $"Given url targets server {artifactUrl.Host} rather than configured server {httpClient.BaseUrl.Host}";
+						return false;
+					}
+
+					GetArtifactResponse response = await HordeHttpRequest.GetAsync<GetArtifactResponse>(httpClient.HttpClient, artifactUrl.PathAndQuery, CancellationToken.None);
+
+					Uri baseUrl = new Uri(httpClient.BaseUrl, artifactUrl.PathAndQuery);
+					Artifact = new ArtifactInfo(response.Name.ToString(), response.Type.ToString(), response.Description, null, response.Keys, response.Metadata, baseUrl, "default");
+
+					ErrorMessage = String.Empty;
+					return true;
+				}
+			}
+			catch (Exception ex)
+			{
+				ErrorMessage = ex.Message;
+				return false;
+			}
+			finally
+			{
+				ShowUrlSpinner = false;
+			}
 		}
 
 		public async Task<bool> SelectUrlAsync()
 		{
+			string text = ArtifactUrl?.ToString() ?? String.Empty;
+
 			TextBox urlTextBox = new TextBox();
-			urlTextBox.Text = InputFile;
-			urlTextBox.SelectionStart = InputFile?.Length ?? 0;
-			urlTextBox.SelectionEnd = InputFile?.Length ?? 0;
+			urlTextBox.Text = text;
+			urlTextBox.SelectionStart = text.Length;
+			urlTextBox.SelectionEnd = text.Length;
 
 			ContentDialog dialog = new ContentDialog()
 			{
@@ -114,40 +216,8 @@ namespace UnrealToolbox.Plugins.Artifacts
 				return false;
 			}
 
-			try
-			{
-				Uri url = new Uri(urlTextBox.Text);
-				using IHordeClientRef clientRef = _hordeClientProvider!.GetClientRef();
-				using HordeHttpClient httpClient = clientRef.Client.CreateHttpClient();
-
-				string path = url.ToString();
-				if (url.GetLeftPart(UriPartial.Authority) == httpClient.BaseUrl.GetLeftPart(UriPartial.Authority))
-				{
-					path = url.PathAndQuery;
-				}
-
-				GetArtifactResponse response = await HordeHttpRequest.GetAsync<GetArtifactResponse>(httpClient.HttpClient, path, CancellationToken.None);
-
-				ArtifactName = response.Name.ToString();
-				ArtifactType = response.Type.ToString();
-				ArtifactDescription = response.Description ?? String.Empty;
-				ArtifactJobUrl = String.Empty;
-
-				List<string> keys = (response.Keys ?? new List<string>()).Where(x => !x.StartsWith("job:", StringComparison.OrdinalIgnoreCase)).ToList();
-				ArtifactKeys = (keys.Count == 0) ? "(None)" : String.Join(", ", keys);
-				ArtifactMetadata = (response.Metadata == null) ? "(None)" : String.Join(", ", response.Metadata);
-
-				InputFile = url.ToString();
-				ShowArtifactInfo = true;
-				ErrorMessage = String.Empty;
-
-				return true;
-			}
-			catch (Exception ex)
-			{
-				ErrorMessage = ex.Message;
-				return false;
-			}
+			ArtifactUrl = new Uri(urlTextBox.Text);
+			return await RefreshArtifactAsync();
 		}
 
 		public async Task<bool> SelectInputFileAsync()
@@ -168,37 +238,8 @@ namespace UnrealToolbox.Plugins.Artifacts
 				return false;
 			}
 
-			InputFile = localPath;
-
-			FileReference file = new FileReference(InputFile);
-			if (!FileReference.Exists(file))
-			{
-				ErrorMessage = "Artifact descriptor not found.";
-				return false;
-			}
-
-			try
-			{
-				Descriptor = await ArtifactDescriptor.ReadAsync(file, CancellationToken.None);
-
-				ArtifactName = Descriptor.Name ?? String.Empty;
-				ArtifactType = Descriptor.Type ?? String.Empty;
-				ArtifactDescription = Descriptor.Description ?? String.Empty;
-				ArtifactJobUrl = Descriptor.JobUrl?.ToString() ?? String.Empty;
-
-				List<string> keys = (Descriptor.Keys ?? new List<string>()).Where(x => !x.StartsWith("job:", StringComparison.OrdinalIgnoreCase)).ToList();
-				ArtifactKeys = (keys.Count == 0) ? "(None)" : String.Join(", ", keys);
-				ArtifactMetadata = (Descriptor.Metadata == null) ? "(None)" : String.Join(", ", Descriptor.Metadata);
-
-				ShowArtifactInfo = true;
-				ErrorMessage = String.Empty;
-				return true;
-			}
-			catch (Exception ex)
-			{
-				ErrorMessage = ex.ToString();
-				return false;
-			}
+			ArtifactUrl = new Uri(new Uri("file://"), localPath);
+			return await RefreshArtifactAsync();
 		}
 
 		public void OpenJobUrl()
diff --git a/Engine/Source/Programs/UnrealToolbox/Plugins/Artifacts/DownloadOptionsWindow.axaml b/Engine/Source/Programs/UnrealToolbox/Plugins/Artifacts/DownloadOptionsWindow.axaml
--- a/Engine/Source/Programs/UnrealToolbox/Plugins/Artifacts/DownloadOptionsWindow.axaml
+++ b/Engine/Source/Programs/UnrealToolbox/Plugins/Artifacts/DownloadOptionsWindow.axaml
@@ -23,10 +23,10 @@
 
 		<Border Classes="Card" CornerRadius="5 5 0 0">
 			<StackPanel Spacing="5" Margin="0">
-				<Grid ColumnDefinitions="*,Auto,Auto" RowDefinitions="Auto">
-					<TextBlock Foreground="Gray" Text="No file selected" IsVisible="{Binding !ShowInputFile}" VerticalAlignment="Center"/>
-					<TextBlock Text="{Binding InputFile}" IsVisible="{Binding ShowInputFile}" VerticalAlignment="Center"/>
-					<Button Content="Browse..." Margin="16 0 0 0" Width="120" Grid.Column="1" Grid.RowSpan="2" Command="{Binding SelectInputFileAsync}"/>
+				<Grid ColumnDefinitions="Auto,*,Auto" RowDefinitions="Auto">
+					<ui:ProgressRing Width="16" Height="16" IsVisible="{Binding ShowUrlSpinner}" VerticalAlignment="Center" Grid.Row="0" Grid.Column="0" Margin="0 0 8 0"/>
+					<SelectableTextBlock Text="{Binding InputText}" VerticalAlignment="Center" Grid.Column="1"/>
+					<Button Content="Browse..." Margin="16 0 0 0" Width="120" Grid.Column="2" Command="{Binding SelectInputFileAsync}" IsVisible="{Binding ShowBrowse}"/>
 				</Grid>
 			</StackPanel>
 		</Border>
@@ -58,7 +58,7 @@
 			<StackPanel Spacing="5" Margin="0">
 				<Grid ColumnDefinitions="*,Auto" Grid.Column="2" Grid.Row="1">
 					<HyperlinkButton Content="{Binding OutputDir}" Command="{Binding OpenOutputDir}" Margin="0" Padding="0"/>
-					<Button Margin="16 0 0 0" Grid.Column="1" Command="{Binding SelectInputFileAsync}" Width="120">Browse...</Button>
+					<Button Margin="16 0 0 0" Grid.Column="1" Command="{Binding SelectOutputDirAsync}" Width="120">Browse...</Button>
 				</Grid>
 			</StackPanel>
 		</Border>
@@ -67,7 +67,7 @@
 			<TextBlock Classes="SettingHeader">Error</TextBlock>
 
 			<Border Classes="Card">
-				<TextBlock Text="{Binding ErrorMessage}" Foreground="Red"/>
+				<TextBlock Text="{Binding ErrorMessage}" Foreground="Red" TextWrapping="Wrap"/>
 			</Border>
 		</StackPanel>
 
diff --git a/Engine/Source/Programs/UnrealToolbox/Plugins/Artifacts/DownloadOptionsWindow.axaml.cs b/Engine/Source/Programs/UnrealToolbox/Plugins/Artifacts/DownloadOptionsWindow.axaml.cs
--- a/Engine/Source/Programs/UnrealToolbox/Plugins/Artifacts/DownloadOptionsWindow.axaml.cs
+++ b/Engine/Source/Programs/UnrealToolbox/Plugins/Artifacts/DownloadOptionsWindow.axaml.cs
@@ -16,14 +16,24 @@ namespace UnrealToolbox.Plugins.Artifacts
 			: this(null, null)
 		{ }
 
-		public DownloadOptionsWindow(string? artifactFile, IHordeClientProvider? hordeClientProvider)
+		public DownloadOptionsWindow(Uri? artifactUrl, IHordeClientProvider? hordeClientProvider)
 		{
-			InitializeComponent();
-
 			_viewModel = new DownloadOptionsViewModel(this, hordeClientProvider);
-			_viewModel.InputFile = artifactFile;
+			_viewModel.ArtifactUrl = artifactUrl;
 
 			DataContext = _viewModel;
+
+			InitializeComponent();
+		}
+
+		protected override void OnOpened(EventArgs e)
+		{
+			base.OnOpened(e);
+
+			if (_viewModel != null && _viewModel.ArtifactUrl != null)
+			{
+				_ = _viewModel.RefreshArtifactAsync();
+			}
 		}
 
 		public async Task SelectInputFileAsync()
@@ -33,25 +43,12 @@ namespace UnrealToolbox.Plugins.Artifacts
 
 		void Start_Click(object? sender, RoutedEventArgs args)
 		{
-			DownloadOptions? options;
-			try
-			{
-				options = _viewModel.GetDownloadOptions();
-			}
-			catch (Exception ex)
-			{
-				_viewModel.ErrorMessage = $"Unable to read artifact descriptor from {_viewModel.InputFile}: {ex.Message}";
-				return;
-			}
-
-			if (options == null)
+			DownloadOptions? options = _viewModel.GetDownloadOptions();
+			if (options != null)
 			{
-				_viewModel.ErrorMessage = $"Unable to read artifact descriptor from {_viewModel.InputFile}";
-				return;
+				_result.TrySetResult(options);
+				Close();
 			}
-
-			_result.TrySetResult(options);
-			Close();
 		}
 
 		void Cancel_Click(object? sender, RoutedEventArgs args)
diff --git a/Engine/Source/Programs/UnrealToolbox/Plugins/HordeAgent/HordeAgentPlugin.cs b/Engine/Source/Programs/UnrealToolbox/Plugins/HordeAgent/HordeAgentPlugin.cs
--- a/Engine/Source/Programs/UnrealToolbox/Plugins/HordeAgent/HordeAgentPlugin.cs
+++ b/Engine/Source/Programs/UnrealToolbox/Plugins/HordeAgent/HordeAgentPlugin.cs
@@ -15,7 +15,7 @@ using System.Text.Json.Serialization;
 
 namespace UnrealToolbox.Plugins.HordeAgent
 {
-	class HordeAgentPlugin : IToolboxPlugin
+	class HordeAgentPlugin : ToolboxPluginBase
 	{
 		record struct IdleStat(string Name, long Value, long MinValue);
 
@@ -39,16 +39,16 @@ namespace UnrealToolbox.Plugins.HordeAgent
 		ToolboxPluginStatus? _status;
 		ToolboxPluginStatus? _reportStatus; // Updated with _status when nothing is currently being installed
 
-		public string Name => "Horde Agent";
+		public override string Name => "Horde Agent";
 
-		public IconSource Icon => new SymbolIconSource() { Symbol = Symbol.People };
+		public override IconSource Icon => new SymbolIconSource() { Symbol = Symbol.People };
 
 		public bool IsEnabled { get; private set; }
 
-		public bool HasSettingsPage()
+		public override bool HasSettingsPage()
 			=> IsEnabled;
 
-		public Control CreateSettingsPage(SettingsContext context)
+		public override Control CreateSettingsPage(SettingsContext context)
 			=> new HordeAgentSettingsPage(context, this);
 
 		public void UpdateSettings(HordeAgentSettings settings)
@@ -91,7 +91,7 @@ namespace UnrealToolbox.Plugins.HordeAgent
 			_tickPauseStateTask = BackgroundTask.StartNew(ctx => TickPauseStateAsync(ctx));
 		}
 
-		public bool Refresh()
+		public override bool Refresh()
 		{
 			if (OperatingSystem.IsWindows())
 			{
@@ -195,7 +195,7 @@ namespace UnrealToolbox.Plugins.HordeAgent
 			};
 		}
 
-		public void PopulateContextMenu(NativeMenu contextMenu)
+		public override void PopulateContextMenu(NativeMenu contextMenu)
 		{
 			if (IsEnabled)
 			{
@@ -205,8 +205,10 @@ namespace UnrealToolbox.Plugins.HordeAgent
 			}
 		}
 
-		public async ValueTask DisposeAsync()
+		public override async ValueTask DisposeAsync()
 		{
+			await base.DisposeAsync();
+
 			await _tickPauseStateTask.DisposeAsync();
 			await _clientTask.DisposeAsync();
 		}
@@ -269,7 +271,7 @@ namespace UnrealToolbox.Plugins.HordeAgent
 			_host.UpdateStatus();
 		}
 
-		public ToolboxPluginStatus GetStatus()
+		public override ToolboxPluginStatus GetStatus()
 		{
 			if (_settings.Mode == AgentMode.Disabled)
 			{
diff --git a/Engine/Source/Programs/UnrealToolbox/Plugins/HordeProxy/HordeProxyPlugin.cs b/Engine/Source/Programs/UnrealToolbox/Plugins/HordeProxy/HordeProxyPlugin.cs
--- a/Engine/Source/Programs/UnrealToolbox/Plugins/HordeProxy/HordeProxyPlugin.cs
+++ b/Engine/Source/Programs/UnrealToolbox/Plugins/HordeProxy/HordeProxyPlugin.cs
@@ -13,7 +13,7 @@ using Microsoft.Extensions.Hosting;
 
 namespace UnrealToolbox.Plugins.HordeProxy
 {
-	class HordeProxyPlugin : IToolboxPlugin
+	class HordeProxyPlugin : ToolboxPluginBase
 	{
 		readonly BackgroundTask _serverTask;
 		readonly IHordeClientProvider _hordeClientProvider;
@@ -27,11 +27,11 @@ namespace UnrealToolbox.Plugins.HordeProxy
 
 		public event Action? OnStatusChanged;
 
-		public string Name => "Horde Proxy";
+		public override string Name => "Horde Proxy";
 
 		public string Status => _status;
 
-		public IconSource Icon => new SymbolIconSource() { Symbol = Symbol.Filter };
+		public override IconSource Icon => new SymbolIconSource() { Symbol = Symbol.Filter };
 
 		public HordeProxyPlugin(IHordeClientProvider hordeClientProvider)
 		{
@@ -47,10 +47,10 @@ namespace UnrealToolbox.Plugins.HordeProxy
 			_serverTask = BackgroundTask.StartNew(RunServerAsync);
 		}
 
-		public bool HasSettingsPage()
+		public override bool HasSettingsPage()
 			=> true;
 
-		public Control CreateSettingsPage(SettingsContext context)
+		public override Control CreateSettingsPage(SettingsContext context)
 			=> new HordeProxySettingsPage(this);
 
 		void SetStatus(string status)
@@ -75,13 +75,15 @@ namespace UnrealToolbox.Plugins.HordeProxy
 			}
 		}
 
-		public async ValueTask DisposeAsync()
+		public override async ValueTask DisposeAsync()
 		{
+			await base.DisposeAsync();
+
 			_hordeClientProvider.OnStateChanged -= OnStateChanged;
 			await _serverTask.DisposeAsync();
 		}
 
-		public bool Refresh()
+		public override bool Refresh()
 			=> false;
 
 		void OnStateChanged()
@@ -89,13 +91,6 @@ namespace UnrealToolbox.Plugins.HordeProxy
 			_restartServerEvent.Pulse();
 		}
 
-		public void PopulateContextMenu(NativeMenu contextMenu)
-		{
-		}
-
-		public ToolboxPluginStatus GetStatus()
-			=> ToolboxPluginStatus.Default;
-
 		async Task RunServerAsync(CancellationToken cancellationToken)
 		{
 			while (!cancellationToken.IsCancellationRequested)
diff --git a/Engine/Source/Programs/UnrealToolbox/Program.cs b/Engine/Source/Programs/UnrealToolbox/Program.cs
--- a/Engine/Source/Programs/UnrealToolbox/Program.cs
+++ b/Engine/Source/Programs/UnrealToolbox/Program.cs
@@ -12,10 +12,13 @@ namespace UnrealToolbox
 		public const string CloseEventName = "UnrealToolbox-Close";
 		public const string RefreshEventName = "UnrealToolbox-Refresh"; // Note: this event is set by InstallerCustomActions when agent is installed.
 		public const string SettingsEventName = "UnrealToolbox-Settings";
+		public const string CommandEventName = "UnrealToolbox-Command";
 
 		public static SelfUpdateState? Update { get; private set; }
 
-		public static FileReference? LogFile { get; } = GetLogFile();
+		public static DirectoryReference DataDir { get; } = GetDataDir();
+
+		public static FileReference LogFile { get; } = FileReference.Combine(DataDir, "Log.txt");
 
 		[STAThread]
 		public static int Main(string[] args)
@@ -62,16 +65,24 @@ namespace UnrealToolbox
 				settingsEvent.Set();
 			}
 
+			using EventWaitHandle commandEvent = new EventWaitHandle(false, EventResetMode.AutoReset, CommandEventName);
+			foreach (string arg in args)
+			{
+				const string UrlArgument = "-Url=";
+				if (arg.StartsWith(UrlArgument, StringComparison.OrdinalIgnoreCase))
+				{
+					CommandHelper.Enqueue(new Command { Type = CommandType.OpenUrl, Argument = arg.Substring(UrlArgument.Length) });
+					commandEvent.Set();
+				}
+			}
+
 			using SingleInstanceMutex mutex = new SingleInstanceMutex(MutexName);
 			if (!mutex.Wait(0))
 			{
 				return 1;
 			}
 
-			if (LogFile != null)
-			{
-				Log.AddFileWriter("Default", LogFile);
-			}
+			Log.AddFileWriter("Default", LogFile);
 
 			AppBuilder builder = BuildAvaloniaApp();
 			try
@@ -97,14 +108,17 @@ namespace UnrealToolbox
 				.LogToTrace();
 		}
 
-		static FileReference? GetLogFile()
+		static DirectoryReference GetDataDir()
 		{
 			DirectoryReference? appDataDir = DirectoryReference.GetSpecialFolder(Environment.SpecialFolder.LocalApplicationData);
 			if (appDataDir == null)
 			{
-				return null;
+				return DirectoryReference.GetCurrentDirectory();
+			}
+			else
+			{
+				return DirectoryReference.Combine(appDataDir, "Epic Games", "Unreal Toolbox");
 			}
-			return FileReference.Combine(appDataDir, "Epic Games", "Unreal Toolbox", "Log.txt");
 		}
 	}
 
