Unreal Toolbox: Generate a job URL for artifact downloads.

#rnx

diff --git a/Engine/Source/Programs/UnrealToolbox/Plugins/Artifacts/DownloadOptionsViewModel.cs b/Engine/Source/Programs/UnrealToolbox/Plugins/Artifacts/DownloadOptionsViewModel.cs
--- a/Engine/Source/Programs/UnrealToolbox/Plugins/Artifacts/DownloadOptionsViewModel.cs
+++ b/Engine/Source/Programs/UnrealToolbox/Plugins/Artifacts/DownloadOptionsViewModel.cs
@@ -1,6 +1,7 @@
 // Copyright Epic Games, Inc. All Rights Reserved.
 
 using System.Diagnostics;
+using System.Text.RegularExpressions;
 using Avalonia.Controls;
 using Avalonia.Platform.Storage;
 using CommunityToolkit.Mvvm.ComponentModel;
@@ -175,8 +176,19 @@ namespace UnrealToolbox.Plugins.Artifacts
 
 					GetArtifactResponse response = await HordeHttpRequest.GetAsync<GetArtifactResponse>(httpClient.HttpClient, artifactUrl.PathAndQuery, CancellationToken.None);
 
+					Uri? jobUrl = null;
+					foreach (string key in response.Keys)
+					{
+						Match match = Regex.Match(key, "^job:([0-9a-fA-F]+)$");
+						if (match.Success)
+						{
+							jobUrl = new Uri(httpClient.BaseUrl, $"job/{match.Groups[1].Value}");
+							break;
+						}
+					}
+
 					Uri baseUrl = new Uri(httpClient.BaseUrl, artifactUrl.PathAndQuery);
-					Artifact = new ArtifactInfo(response.Name.ToString(), response.Type.ToString(), response.Description, null, response.Keys, response.Metadata, baseUrl, "default");
+					Artifact = new ArtifactInfo(response.Name.ToString(), response.Type.ToString(), response.Description, jobUrl, response.Keys, response.Metadata, baseUrl, "default");
 
 					ErrorMessage = String.Empty;
 					return true;
