Unreal Toolbox: Save the output directory for downloads.

#rnx

diff --git a/Engine/Source/Programs/UnrealToolbox/AboutPage.axaml.cs b/Engine/Source/Programs/UnrealToolbox/AboutPage.axaml.cs
--- a/Engine/Source/Programs/UnrealToolbox/AboutPage.axaml.cs
+++ b/Engine/Source/Programs/UnrealToolbox/AboutPage.axaml.cs
@@ -15,7 +15,7 @@ namespace UnrealToolbox
 			string version = Program.Update?.CurrentVersionString ?? "No version information present.";
 			_versionText.Text = version;
 			
-			string settingsDir = ToolCatalog.GetBaseDir().FullName;
+			string settingsDir = Program.DataDir.FullName;
 			_settingsDirButton.Click += (sender, args) => OpenFolder(settingsDir);
 			_settingsDirText.Text = settingsDir;
 			
diff --git a/Engine/Source/Programs/UnrealToolbox/Plugins/Artifacts/ArtifactsPlugin.cs b/Engine/Source/Programs/UnrealToolbox/Plugins/Artifacts/ArtifactsPlugin.cs
--- a/Engine/Source/Programs/UnrealToolbox/Plugins/Artifacts/ArtifactsPlugin.cs
+++ b/Engine/Source/Programs/UnrealToolbox/Plugins/Artifacts/ArtifactsPlugin.cs
@@ -14,6 +14,7 @@ namespace UnrealToolbox.Plugins.Artifacts
 	class ArtifactsPlugin : ToolboxPluginBase
 	{
 		readonly IHordeClientProvider _hordeClientProvider;
+		readonly JsonConfig<DownloadSettings> _settings;
 		readonly ILogger _logger;
 
 		public override string Name
@@ -28,6 +29,8 @@ namespace UnrealToolbox.Plugins.Artifacts
 		public ArtifactsPlugin(IHordeClientProvider hordeClientProvider, ILogger<ArtifactsPlugin> logger)
 		{
 			_hordeClientProvider = hordeClientProvider;
+			_settings = new JsonConfig<DownloadSettings>(FileReference.Combine(Program.DataDir, "Artifacts.json"));
+			_settings.LoadSettings();
 			_logger = logger;
 		}
 
@@ -56,7 +59,7 @@ namespace UnrealToolbox.Plugins.Artifacts
 
 		public async Task DownloadAsync(Uri? artifactUrl, CancellationToken cancellationToken)
 		{
-			DownloadOptionsWindow window = new DownloadOptionsWindow(artifactUrl, _hordeClientProvider);
+			DownloadOptionsWindow window = new DownloadOptionsWindow(artifactUrl, _settings, _hordeClientProvider);
 			window.Show();
 
 			DownloadOptions? options = await window.Result;
diff --git a/Engine/Source/Programs/UnrealToolbox/Plugins/Artifacts/DownloadOptionsViewModel.cs b/Engine/Source/Programs/UnrealToolbox/Plugins/Artifacts/DownloadOptionsViewModel.cs
--- a/Engine/Source/Programs/UnrealToolbox/Plugins/Artifacts/DownloadOptionsViewModel.cs
+++ b/Engine/Source/Programs/UnrealToolbox/Plugins/Artifacts/DownloadOptionsViewModel.cs
@@ -64,7 +64,7 @@ namespace UnrealToolbox.Plugins.Artifacts
 			=> ArtifactUrl == null || ArtifactUrl.IsFile;
 
 		[ObservableProperty]
-		string _outputDir = Directory.GetCurrentDirectory();
+		string _outputDir;
 
 		[ObservableProperty]
 		[NotifyPropertyChangedFor(nameof(ShowErrorMessage))]
@@ -115,9 +115,10 @@ namespace UnrealToolbox.Plugins.Artifacts
 
 		readonly IHordeClientProvider? _hordeClientProvider;
 
-		public DownloadOptionsViewModel(Window window, IHordeClientProvider? hordeClientProvider)
+		public DownloadOptionsViewModel(Window window, string? outputDir, IHordeClientProvider? hordeClientProvider)
 		{
 			_window = window;
+			_outputDir = outputDir ?? Directory.GetCurrentDirectory();
 			_hordeClientProvider = hordeClientProvider;
 		}
 
diff --git a/Engine/Source/Programs/UnrealToolbox/Plugins/Artifacts/DownloadOptionsWindow.axaml b/Engine/Source/Programs/UnrealToolbox/Plugins/Artifacts/DownloadOptionsWindow.axaml
--- a/Engine/Source/Programs/UnrealToolbox/Plugins/Artifacts/DownloadOptionsWindow.axaml
+++ b/Engine/Source/Programs/UnrealToolbox/Plugins/Artifacts/DownloadOptionsWindow.axaml
@@ -21,7 +21,7 @@
 
 		<TextBlock Classes="SettingHeader">Source</TextBlock>
 
-		<Border Classes="Card" CornerRadius="5 5 0 0">
+		<Border Classes="Card">
 			<StackPanel Spacing="5" Margin="0">
 				<Grid ColumnDefinitions="Auto,*,Auto" RowDefinitions="Auto">
 					<ui:ProgressRing Width="16" Height="16" IsVisible="{Binding ShowUrlSpinner}" VerticalAlignment="Center" Grid.Row="0" Grid.Column="0" Margin="0 0 8 0"/>
@@ -30,7 +30,7 @@
 				</Grid>
 			</StackPanel>
 		</Border>
-		<Border Classes="Card" CornerRadius="0 0 5 5" BorderThickness="1 0 1 1" IsVisible="{Binding ShowArtifactInfo}">
+		<Border Classes="Card" BorderThickness="1 0 1 1" IsVisible="{Binding ShowArtifactInfo}">
 			<Grid ColumnDefinitions="Auto,*" RowDefinitions="Auto,Auto,Auto,Auto,Auto,Auto">
 				<TextBlock Margin="0 1" Theme="{StaticResource CaptionTextBlockStyle}" Text="Name" Padding="0" Grid.Row="0" IsVisible="{Binding HasArtifactName}"/>
 				<TextBlock Margin="42 1" Theme="{StaticResource CaptionTextBlockStyle}" Text="{Binding ArtifactName}" Padding="0" Grid.Row="0" Grid.Column="1" IsVisible="{Binding HasArtifactName}"/>
@@ -54,7 +54,7 @@
 
 		<TextBlock Classes="SettingHeader">Output</TextBlock>
 
-		<Border Classes="Card" CornerRadius="10 0 10 0">
+		<Border Classes="Card">
 			<StackPanel Spacing="5" Margin="0">
 				<Grid ColumnDefinitions="*,Auto" Grid.Column="2" Grid.Row="1">
 					<HyperlinkButton Content="{Binding OutputDir}" Command="{Binding OpenOutputDir}" Margin="0" Padding="0"/>
diff --git a/Engine/Source/Programs/UnrealToolbox/Plugins/Artifacts/DownloadOptionsWindow.axaml.cs b/Engine/Source/Programs/UnrealToolbox/Plugins/Artifacts/DownloadOptionsWindow.axaml.cs
--- a/Engine/Source/Programs/UnrealToolbox/Plugins/Artifacts/DownloadOptionsWindow.axaml.cs
+++ b/Engine/Source/Programs/UnrealToolbox/Plugins/Artifacts/DownloadOptionsWindow.axaml.cs
@@ -8,17 +8,20 @@ namespace UnrealToolbox.Plugins.Artifacts
 	partial class DownloadOptionsWindow : Window
 	{
 		readonly DownloadOptionsViewModel _viewModel;
+		readonly JsonConfig<DownloadSettings>? _settings;
 		readonly TaskCompletionSource<DownloadOptions?> _result = new TaskCompletionSource<DownloadOptions?>();
 
 		public Task<DownloadOptions?> Result => _result.Task;
 
 		public DownloadOptionsWindow()
-			: this(null, null)
+			: this(null, null, null)
 		{ }
 
-		public DownloadOptionsWindow(Uri? artifactUrl, IHordeClientProvider? hordeClientProvider)
+		public DownloadOptionsWindow(Uri? artifactUrl, JsonConfig<DownloadSettings>? settings, IHordeClientProvider? hordeClientProvider)
 		{
-			_viewModel = new DownloadOptionsViewModel(this, hordeClientProvider);
+			_settings = settings;
+
+			_viewModel = new DownloadOptionsViewModel(this, _settings?.Current.OutputDir, hordeClientProvider);
 			_viewModel.ArtifactUrl = artifactUrl;
 
 			DataContext = _viewModel;
@@ -46,6 +49,13 @@ namespace UnrealToolbox.Plugins.Artifacts
 			DownloadOptions? options = _viewModel.GetDownloadOptions();
 			if (options != null)
 			{
+				if (_settings != null)
+				{
+					DownloadSettings newSettings = new DownloadSettings(_settings.Current);
+					newSettings.OutputDir = options.OutputDir?.ToString();
+					_settings.UpdateSettings(newSettings);
+				}
+
 				_result.TrySetResult(options);
 				Close();
 			}
diff --git a/Engine/Source/Programs/UnrealToolbox/Plugins/Artifacts/DownloadSettings.cs b/Engine/Source/Programs/UnrealToolbox/Plugins/Artifacts/DownloadSettings.cs
new file mode 100644
--- /dev/null
+++ b/Engine/Source/Programs/UnrealToolbox/Plugins/Artifacts/DownloadSettings.cs
@@ -0,0 +1,30 @@
+// Copyright Epic Games, Inc. All Rights Reserved.
+
+namespace UnrealToolbox.Plugins.Artifacts
+{
+	/// <summary>
+	/// Settings for the artifacts plugin
+	/// </summary>
+	[Serializable]
+	public class DownloadSettings
+	{
+		/// <summary>
+		/// Output directory for the download
+		/// </summary>
+		public string? OutputDir { get; set; }
+
+		/// <summary>
+		/// Default constructor
+		/// </summary>
+		public DownloadSettings()
+		{ }
+
+		/// <summary>
+		/// Copy constructor
+		/// </summary>
+		public DownloadSettings(DownloadSettings other)
+		{
+			OutputDir = other.OutputDir;
+		}
+	}
+}
diff --git a/Engine/Source/Programs/UnrealToolbox/Plugins/HordeAgent/HordeAgentPlugin.cs b/Engine/Source/Programs/UnrealToolbox/Plugins/HordeAgent/HordeAgentPlugin.cs
--- a/Engine/Source/Programs/UnrealToolbox/Plugins/HordeAgent/HordeAgentPlugin.cs
+++ b/Engine/Source/Programs/UnrealToolbox/Plugins/HordeAgent/HordeAgentPlugin.cs
@@ -81,9 +81,7 @@ namespace UnrealToolbox.Plugins.HordeAgent
 			_toolCatalog = toolCatalog;
 			_logger = logger;
 
-			DirectoryReference? settingsRoot = DirectoryReference.GetSpecialFolder(Environment.SpecialFolder.LocalApplicationData);
-			settingsRoot ??= DirectoryReference.GetCurrentDirectory();
-			_settingsFile = FileReference.Combine(settingsRoot, "Epic Games", "Unreal Toolbox", "HordeAgent.json");
+			_settingsFile = FileReference.Combine(Program.DataDir, "HordeAgent.json");
 
 			LoadSettings();
 
diff --git a/Engine/Source/Programs/UnrealToolbox/Plugins/HordeProxy/HordeProxyPlugin.cs b/Engine/Source/Programs/UnrealToolbox/Plugins/HordeProxy/HordeProxyPlugin.cs
--- a/Engine/Source/Programs/UnrealToolbox/Plugins/HordeProxy/HordeProxyPlugin.cs
+++ b/Engine/Source/Programs/UnrealToolbox/Plugins/HordeProxy/HordeProxyPlugin.cs
@@ -41,7 +41,7 @@ namespace UnrealToolbox.Plugins.HordeProxy
 
 			DirectoryReference? settingsRoot = DirectoryReference.GetSpecialFolder(Environment.SpecialFolder.LocalApplicationData);
 			settingsRoot ??= DirectoryReference.GetCurrentDirectory();
-			_settings = new JsonConfig<HordeProxySettings>(FileReference.Combine(settingsRoot, "Epic Games", "Unreal Toolbox", "HordeProxy.json"));
+			_settings = new JsonConfig<HordeProxySettings>(FileReference.Combine(Program.DataDir, "HordeProxy.json"));
 			_settings.LoadSettings();
 
 			_serverTask = BackgroundTask.StartNew(RunServerAsync);
diff --git a/Engine/Source/Programs/UnrealToolbox/ToolCatalog.cs b/Engine/Source/Programs/UnrealToolbox/ToolCatalog.cs
--- a/Engine/Source/Programs/UnrealToolbox/ToolCatalog.cs
+++ b/Engine/Source/Programs/UnrealToolbox/ToolCatalog.cs
@@ -254,15 +254,8 @@ namespace UnrealToolbox
 			UpdateMainThread();
 		}
 		
-		public static DirectoryReference GetBaseDir()
-		{
-			DirectoryReference? dir = DirectoryReference.GetSpecialFolder(Environment.SpecialFolder.LocalApplicationData);
-			dir ??= DirectoryReference.GetCurrentDirectory();
-			return DirectoryReference.Combine(dir, "Epic Games", "Unreal Toolbox");
-		}
-		
-		public static DirectoryReference GetToolsDir() => DirectoryReference.Combine(GetBaseDir(), "Tools");
-		public static FileReference GetConfigFile() => FileReference.Combine(GetBaseDir(), "Tools.json");
+		public static DirectoryReference GetToolsDir() => DirectoryReference.Combine(Program.DataDir, "Tools");
+		public static FileReference GetConfigFile() => FileReference.Combine(Program.DataDir, "Tools.json");
 
 		void PostMainThreadUpdate()
 		{
