Unreal Toolbox: Force the progress bar to complete once a download finishes.

#rnx

diff --git a/Engine/Source/Programs/UnrealToolbox/Plugins/Artifacts/ArtifactsPlugin.cs b/Engine/Source/Programs/UnrealToolbox/Plugins/Artifacts/ArtifactsPlugin.cs
--- a/Engine/Source/Programs/UnrealToolbox/Plugins/Artifacts/ArtifactsPlugin.cs
+++ b/Engine/Source/Programs/UnrealToolbox/Plugins/Artifacts/ArtifactsPlugin.cs
@@ -107,6 +107,7 @@ namespace UnrealToolbox.Plugins.Artifacts
 					await blobRef.ExtractAsync(options.OutputDir.ToDirectoryInfo(), new ExtractOptions { Progress = progress, ProgressUpdateFrequency = TimeSpan.FromSeconds(0.2) }, _logger, linkedCancellationToken);
 
 					progress.Status = "Download complete.";
+					progress.MarkComplete();
 				}
 
 				progress.OpenFolderPath = options.OutputDir.ToString();
diff --git a/Engine/Source/Programs/UnrealToolbox/Plugins/Artifacts/DownloadProgressViewModel.cs b/Engine/Source/Programs/UnrealToolbox/Plugins/Artifacts/DownloadProgressViewModel.cs
--- a/Engine/Source/Programs/UnrealToolbox/Plugins/Artifacts/DownloadProgressViewModel.cs
+++ b/Engine/Source/Programs/UnrealToolbox/Plugins/Artifacts/DownloadProgressViewModel.cs
@@ -1,5 +1,6 @@
 // Copyright Epic Games, Inc. All Rights Reserved.
 
+using Avalonia.Threading;
 using CommunityToolkit.Mvvm.ComponentModel;
 using EpicGames.Horde.Storage.Nodes;
 
@@ -42,9 +43,21 @@ namespace UnrealToolbox.Plugins.Artifacts
 
 		public void Report(IExtractStats value)
 		{
-			Status = $"Copied {value.NumFiles} files ({value.DownloadSize / (1024.0 * 1024.0):n1}mb, {value.DownloadRate / (1024.0 * 1024.0):n1}mb/s)...";
+			string status = $"Copied {value.NumFiles} files ({value.DownloadSize / (1024.0 * 1024.0):n1}mb, {value.DownloadRate / (1024.0 * 1024.0):n1}mb/s)...";
+			long size = value.DownloadSize;
+			Dispatcher.UIThread.Invoke(() => ReportOnMainThread(status, size));
+		}
+
+		public void MarkComplete()
+		{
+			_size = _totalSize;
+			UpdateProgress();
+		}
 
-			_size = value.DownloadSize;
+		void ReportOnMainThread(string status, long size)
+		{
+			Status = status;
+			_size = size;
 			UpdateProgress();
 		}
 
