Unreal Toolbox: Add an option to append the build name to the output directory for downloaded artifacts.

diff --git a/Engine/Source/Programs/UnrealToolbox/Plugins/Artifacts/DownloadOptionsViewModel.cs b/Engine/Source/Programs/UnrealToolbox/Plugins/Artifacts/DownloadOptionsViewModel.cs
--- a/Engine/Source/Programs/UnrealToolbox/Plugins/Artifacts/DownloadOptionsViewModel.cs
+++ b/Engine/Source/Programs/UnrealToolbox/Plugins/Artifacts/DownloadOptionsViewModel.cs
@@ -1,6 +1,8 @@
 // Copyright Epic Games, Inc. All Rights Reserved.
 
 using System.Diagnostics;
+using System.Diagnostics.CodeAnalysis;
+using System.Text;
 using System.Text.RegularExpressions;
 using Avalonia.Controls;
 using Avalonia.Platform.Storage;
@@ -8,12 +10,14 @@ using CommunityToolkit.Mvvm.ComponentModel;
 using EpicGames.Core;
 using EpicGames.Horde;
 using EpicGames.Horde.Artifacts;
+using EpicGames.Horde.Commits;
 using EpicGames.Horde.Storage;
+using EpicGames.Horde.Streams;
 using FluentAvalonia.UI.Controls;
 
 namespace UnrealToolbox.Plugins.Artifacts
 {
-	record class ArtifactInfo(string? Name, string? Type, string? Description, Uri? JobUrl, IReadOnlyList<string>? Keys, IReadOnlyList<string>? Metadata, Uri BaseUrl, RefName RefName);
+	record class ArtifactInfo(ArtifactId? Id, string? Name, string? Type, StreamId? Stream, CommitId? Commit, string? Description, Uri? JobUrl, IReadOnlyList<string>? Keys, IReadOnlyList<string>? Metadata, Uri BaseUrl, RefName RefName);
 
 	partial class DownloadOptionsViewModel : ObservableObject
 	{
@@ -28,6 +32,7 @@ namespace UnrealToolbox.Plugins.Artifacts
 
 		[ObservableProperty]
 		[NotifyPropertyChangedFor(nameof(AllowStart))]
+		[NotifyPropertyChangedFor(nameof(BuildName))]
 		[NotifyPropertyChangedFor(nameof(ShowArtifactInfo))]
 		[NotifyPropertyChangedFor(nameof(HasArtifactName))]
 		[NotifyPropertyChangedFor(nameof(ArtifactName))]
@@ -39,6 +44,7 @@ namespace UnrealToolbox.Plugins.Artifacts
 		[NotifyPropertyChangedFor(nameof(ArtifactJobUrl))]
 		[NotifyPropertyChangedFor(nameof(ArtifactKeys))]
 		[NotifyPropertyChangedFor(nameof(ArtifactMetadata))]
+		[NotifyPropertyChangedFor(nameof(OutputDir))]
 		ArtifactInfo? _artifact;
 
 		public string InputText
@@ -63,8 +69,16 @@ namespace UnrealToolbox.Plugins.Artifacts
 		public bool ShowBrowse
 			=> ArtifactUrl == null || ArtifactUrl.IsFile;
 
+		public string OutputDir
+			=> AppendBuildName ? Path.Combine(BaseOutputDir, BuildName) : BaseOutputDir;
+
+		[ObservableProperty]
+		[NotifyPropertyChangedFor(nameof(OutputDir))]
+		public string _baseOutputDir;
+
 		[ObservableProperty]
-		string _outputDir;
+		[NotifyPropertyChangedFor(nameof(OutputDir))]
+		public bool _appendBuildName;
 
 		[ObservableProperty]
 		[NotifyPropertyChangedFor(nameof(ShowErrorMessage))]
@@ -113,12 +127,28 @@ namespace UnrealToolbox.Plugins.Artifacts
 		public bool AllowStart
 			=> Artifact != null;
 
+		public string BuildName
+		{
+			get
+			{
+				if (Artifact != null && Artifact.Stream != null && Artifact.Commit != null)
+				{
+					return $"{Artifact.Stream}-{Artifact.Commit}";
+				}
+				else
+				{
+					return "";
+				}
+			}
+		}
+
 		readonly IHordeClientProvider? _hordeClientProvider;
 
-		public DownloadOptionsViewModel(Window window, string? outputDir, IHordeClientProvider? hordeClientProvider)
+		public DownloadOptionsViewModel(Window window, DownloadSettings? settings, IHordeClientProvider? hordeClientProvider)
 		{
 			_window = window;
-			_outputDir = outputDir ?? Directory.GetCurrentDirectory();
+			_baseOutputDir = settings?.OutputDir ?? Directory.GetCurrentDirectory();
+			_appendBuildName = settings?.AppendBuildName ?? true;
 			_hordeClientProvider = hordeClientProvider;
 		}
 
@@ -159,7 +189,7 @@ namespace UnrealToolbox.Plugins.Artifacts
 					}
 
 					ArtifactDescriptor descriptor = await ArtifactDescriptor.ReadAsync(file, CancellationToken.None);
-					Artifact = new ArtifactInfo(descriptor.Name, descriptor.Type, descriptor.Description, descriptor.JobUrl, descriptor.Keys, descriptor.Metadata, descriptor.BaseUrl, descriptor.RefName);
+					Artifact = new ArtifactInfo(null, descriptor.Name, descriptor.Type, null, null, descriptor.Description, descriptor.JobUrl, descriptor.Keys, descriptor.Metadata, descriptor.BaseUrl, descriptor.RefName);
 
 					ErrorMessage = String.Empty;
 					return true;
@@ -189,7 +219,7 @@ namespace UnrealToolbox.Plugins.Artifacts
 					}
 
 					Uri baseUrl = new Uri(httpClient.BaseUrl, artifactUrl.PathAndQuery);
-					Artifact = new ArtifactInfo(response.Name.ToString(), response.Type.ToString(), response.Description, jobUrl, response.Keys, response.Metadata, baseUrl, "default");
+					Artifact = new ArtifactInfo(response.Id, response.Name.ToString(), response.Type.ToString(), response.StreamId, response.CommitId, response.Description, jobUrl, response.Keys, response.Metadata, baseUrl, "default");
 
 					ErrorMessage = String.Empty;
 					return true;
@@ -309,7 +339,7 @@ namespace UnrealToolbox.Plugins.Artifacts
 				return false;
 			}
 
-			OutputDir = localPath;
+			BaseOutputDir = localPath;
 			return true;
 		}
 	}
diff --git a/Engine/Source/Programs/UnrealToolbox/Plugins/Artifacts/DownloadOptionsWindow.axaml b/Engine/Source/Programs/UnrealToolbox/Plugins/Artifacts/DownloadOptionsWindow.axaml
--- a/Engine/Source/Programs/UnrealToolbox/Plugins/Artifacts/DownloadOptionsWindow.axaml
+++ b/Engine/Source/Programs/UnrealToolbox/Plugins/Artifacts/DownloadOptionsWindow.axaml
@@ -55,11 +55,12 @@
 		<TextBlock Classes="SettingHeader">Output</TextBlock>
 
 		<Border Classes="Card">
-			<StackPanel Spacing="5" Margin="0">
+			<StackPanel Spacing="4" Margin="0">
 				<Grid ColumnDefinitions="*,Auto" Grid.Column="2" Grid.Row="1">
 					<HyperlinkButton Content="{Binding OutputDir}" Command="{Binding OpenOutputDir}" Margin="0" Padding="0"/>
 					<Button Margin="16 0 0 0" Grid.Column="1" Command="{Binding SelectOutputDirAsync}" Width="120">Browse...</Button>
 				</Grid>
+				<CheckBox IsChecked="{Binding AppendBuildName}" Content="Append build name"/>
 			</StackPanel>
 		</Border>
 
diff --git a/Engine/Source/Programs/UnrealToolbox/Plugins/Artifacts/DownloadOptionsWindow.axaml.cs b/Engine/Source/Programs/UnrealToolbox/Plugins/Artifacts/DownloadOptionsWindow.axaml.cs
--- a/Engine/Source/Programs/UnrealToolbox/Plugins/Artifacts/DownloadOptionsWindow.axaml.cs
+++ b/Engine/Source/Programs/UnrealToolbox/Plugins/Artifacts/DownloadOptionsWindow.axaml.cs
@@ -21,7 +21,7 @@ namespace UnrealToolbox.Plugins.Artifacts
 		{
 			_settings = settings;
 
-			_viewModel = new DownloadOptionsViewModel(this, _settings?.Current.OutputDir, hordeClientProvider);
+			_viewModel = new DownloadOptionsViewModel(this, _settings?.Current, hordeClientProvider);
 			_viewModel.ArtifactUrl = artifactUrl;
 
 			DataContext = _viewModel;
@@ -53,6 +53,7 @@ namespace UnrealToolbox.Plugins.Artifacts
 				{
 					DownloadSettings newSettings = new DownloadSettings(_settings.Current);
 					newSettings.OutputDir = options.OutputDir?.ToString();
+					newSettings.AppendBuildName = _viewModel.AppendBuildName;
 					_settings.UpdateSettings(newSettings);
 				}
 
diff --git a/Engine/Source/Programs/UnrealToolbox/Plugins/Artifacts/DownloadSettings.cs b/Engine/Source/Programs/UnrealToolbox/Plugins/Artifacts/DownloadSettings.cs
--- a/Engine/Source/Programs/UnrealToolbox/Plugins/Artifacts/DownloadSettings.cs
+++ b/Engine/Source/Programs/UnrealToolbox/Plugins/Artifacts/DownloadSettings.cs
@@ -13,6 +13,11 @@ namespace UnrealToolbox.Plugins.Artifacts
 		/// </summary>
 		public string? OutputDir { get; set; }
 
+		/// <summary>
+		/// Whether to append the build name to the output directory
+		/// </summary>
+		public bool AppendBuildName { get; set; }
+
 		/// <summary>
 		/// Default constructor
 		/// </summary>
@@ -25,6 +30,7 @@ namespace UnrealToolbox.Plugins.Artifacts
 		public DownloadSettings(DownloadSettings other)
 		{
 			OutputDir = other.OutputDir;
+			AppendBuildName = other.AppendBuildName;
 		}
 	}
 }
