Unreal Toolbox: Fix static analysis warning.

#rnx

diff --git a/Engine/Source/Programs/UnrealToolbox/Plugins/Artifacts/DownloadOptionsViewModel.cs b/Engine/Source/Programs/UnrealToolbox/Plugins/Artifacts/DownloadOptionsViewModel.cs
--- a/Engine/Source/Programs/UnrealToolbox/Plugins/Artifacts/DownloadOptionsViewModel.cs
+++ b/Engine/Source/Programs/UnrealToolbox/Plugins/Artifacts/DownloadOptionsViewModel.cs
@@ -1,8 +1,6 @@
 // Copyright Epic Games, Inc. All Rights Reserved.
 
 using System.Diagnostics;
-using System.Diagnostics.CodeAnalysis;
-using System.Text;
 using System.Text.RegularExpressions;
 using Avalonia.Controls;
 using Avalonia.Platform.Storage;
