Unreal Toolbox: Fix download directory being saved to settings with build name attached.

#rnx

diff --git a/Engine/Source/Programs/UnrealToolbox/Plugins/Artifacts/DownloadOptionsWindow.axaml.cs b/Engine/Source/Programs/UnrealToolbox/Plugins/Artifacts/DownloadOptionsWindow.axaml.cs
--- a/Engine/Source/Programs/UnrealToolbox/Plugins/Artifacts/DownloadOptionsWindow.axaml.cs
+++ b/Engine/Source/Programs/UnrealToolbox/Plugins/Artifacts/DownloadOptionsWindow.axaml.cs
@@ -52,7 +52,7 @@ namespace UnrealToolbox.Plugins.Artifacts
 				if (_settings != null)
 				{
 					DownloadSettings newSettings = new DownloadSettings(_settings.Current);
-					newSettings.OutputDir = options.OutputDir?.ToString();
+					newSettings.OutputDir = _viewModel.BaseOutputDir;
 					newSettings.AppendBuildName = _viewModel.AppendBuildName;
 					_settings.UpdateSettings(newSettings);
 				}
