Unreal Toolbox: Fix issues running toolbox with no Horde server configured, and fix some build warnings on MacOS.

#rnx

diff --git a/Engine/Source/Programs/UnrealToolbox/GeneralSettingsViewModel.cs b/Engine/Source/Programs/UnrealToolbox/GeneralSettingsViewModel.cs
--- a/Engine/Source/Programs/UnrealToolbox/GeneralSettingsViewModel.cs
+++ b/Engine/Source/Programs/UnrealToolbox/GeneralSettingsViewModel.cs
@@ -152,7 +152,11 @@ namespace UnrealToolbox
 
 		async Task<ConnectionState> HandleRefreshAsync(CancellationToken cancellationToken)
 		{
-			using IHordeClientRef hordeClientRef = _hordeClientProvider.GetClientRef();
+			using IHordeClientRef? hordeClientRef = _hordeClientProvider.GetClientRef();
+			if (hordeClientRef == null)
+			{
+				return ConnectionState.NoServerConfigured;
+			}
 
 			IHordeClient client = hordeClientRef.Client;
 			try
diff --git a/Engine/Source/Programs/UnrealToolbox/HordeClientProvider.cs b/Engine/Source/Programs/UnrealToolbox/HordeClientProvider.cs
--- a/Engine/Source/Programs/UnrealToolbox/HordeClientProvider.cs
+++ b/Engine/Source/Programs/UnrealToolbox/HordeClientProvider.cs
@@ -85,6 +85,7 @@ namespace UnrealToolbox
 		}
 
 		readonly object _lockObject = new object();
+		readonly ILogger _logger;
 		readonly ILoggerFactory _loggerFactory;
 		readonly List<Task> _disposeTasks = new List<Task>();
 
@@ -105,8 +106,10 @@ namespace UnrealToolbox
 		/// </summary>
 		public HordeClientProvider(ILoggerFactory loggerFactory)
 		{
+			_logger = loggerFactory.CreateLogger<HordeClientProvider>();
 			_loggerFactory = loggerFactory;
-			_lifetime = new HordeClientLifetime(_loggerFactory);
+
+			CreateLifetime();
 		}
 
 		/// <inheritdoc/>
@@ -120,11 +123,18 @@ namespace UnrealToolbox
 		}
 
 		/// <inheritdoc/>
-		public IHordeClientRef GetClientRef()
+		public IHordeClientRef? GetClientRef()
 		{
 			lock (_lockObject)
 			{
-				return new HordeClientRef(_lifetime ?? throw new ObjectDisposedException("HordeClientProvider"));
+				if (_lifetime == null)
+				{
+					return null;
+				}
+				else
+				{
+					return new HordeClientRef(_lifetime);
+				}
 			}
 		}
 
@@ -149,8 +159,16 @@ namespace UnrealToolbox
 		{
 			Debug.Assert(_lifetime == null);
 
-			_lifetime = new HordeClientLifetime(_loggerFactory);
-			_lifetime.Client.OnAccessTokenStateChanged += OnAccessTokenStateChangedInternal;
+			try
+			{
+				_lifetime = new HordeClientLifetime(_loggerFactory);
+				_lifetime.Client.OnAccessTokenStateChanged += OnAccessTokenStateChangedInternal;
+			}
+			catch (Exception ex)
+			{
+				_logger.LogError(ex, "Unable to create Horde client lifetime: {Message}", ex.Message);
+				_lifetime = null;
+			}
 		}
 
 		void DestroyLifetime()
diff --git a/Engine/Source/Programs/UnrealToolbox/IHordeClientProvider.cs b/Engine/Source/Programs/UnrealToolbox/IHordeClientProvider.cs
--- a/Engine/Source/Programs/UnrealToolbox/IHordeClientProvider.cs
+++ b/Engine/Source/Programs/UnrealToolbox/IHordeClientProvider.cs
@@ -38,6 +38,6 @@ namespace UnrealToolbox
 		/// <summary>
 		/// Gets a reference to the current client instance. These references should be kept as short as possible.
 		/// </summary>
-		IHordeClientRef GetClientRef();
+		IHordeClientRef? GetClientRef();
 	}
 }
diff --git a/Engine/Source/Programs/UnrealToolbox/Plugins/Artifacts/ArtifactsPlugin.cs b/Engine/Source/Programs/UnrealToolbox/Plugins/Artifacts/ArtifactsPlugin.cs
--- a/Engine/Source/Programs/UnrealToolbox/Plugins/Artifacts/ArtifactsPlugin.cs
+++ b/Engine/Source/Programs/UnrealToolbox/Plugins/Artifacts/ArtifactsPlugin.cs
@@ -1,5 +1,6 @@
 // Copyright Epic Games, Inc. All Rights Reserved.
 
+using System.Runtime.Versioning;
 using Avalonia.Controls;
 using EpicGames.Core;
 using EpicGames.Horde;
@@ -127,6 +128,7 @@ namespace UnrealToolbox.Plugins.Artifacts
 			}
 		}
 
+		[SupportedOSPlatform("windows")]
 		public static void RegisterFileAssociations(string application)
 		{
 			string progId = "UnrealGameSync.Artifact";
diff --git a/Engine/Source/Programs/UnrealToolbox/Plugins/Artifacts/DownloadOptionsViewModel.cs b/Engine/Source/Programs/UnrealToolbox/Plugins/Artifacts/DownloadOptionsViewModel.cs
--- a/Engine/Source/Programs/UnrealToolbox/Plugins/Artifacts/DownloadOptionsViewModel.cs
+++ b/Engine/Source/Programs/UnrealToolbox/Plugins/Artifacts/DownloadOptionsViewModel.cs
@@ -194,7 +194,13 @@ namespace UnrealToolbox.Plugins.Artifacts
 				}
 				else
 				{
-					using IHordeClientRef clientRef = _hordeClientProvider!.GetClientRef();
+					using IHordeClientRef? clientRef = _hordeClientProvider!.GetClientRef();
+					if (clientRef == null)
+					{
+						ErrorMessage = $"No server is configured.";
+						return false;
+					}
+
 					using HordeHttpClient httpClient = clientRef.Client.CreateHttpClient();
 
 					if (artifactUrl.Host != httpClient.BaseUrl.Host)
diff --git a/Engine/Source/Programs/UnrealToolbox/SelfUpdateService.cs b/Engine/Source/Programs/UnrealToolbox/SelfUpdateService.cs
--- a/Engine/Source/Programs/UnrealToolbox/SelfUpdateService.cs
+++ b/Engine/Source/Programs/UnrealToolbox/SelfUpdateService.cs
@@ -292,7 +292,11 @@ namespace UnrealToolbox
 				return;
 			}
 
-			using IHordeClientRef clientRef = _hordeClientProvider.GetClientRef();
+			using IHordeClientRef? clientRef = _hordeClientProvider.GetClientRef();
+			if (clientRef == null)
+			{
+				return;
+			}
 
 			ITool? tool = await clientRef.Client.Tools.GetAsync(ToolId, cancellationToken);
 			if (tool == null || tool.Deployments.Count == 0)
diff --git a/Engine/Source/Programs/UnrealToolbox/ToolCatalog.cs b/Engine/Source/Programs/UnrealToolbox/ToolCatalog.cs
--- a/Engine/Source/Programs/UnrealToolbox/ToolCatalog.cs
+++ b/Engine/Source/Programs/UnrealToolbox/ToolCatalog.cs
@@ -564,6 +564,11 @@ namespace UnrealToolbox
 
 		static bool IsMsiInstalled(string msiProductId)
 		{
+			if (!OperatingSystem.IsWindows())
+			{
+				return false;
+			}
+
 			string? value =
 				(Registry.GetValue($"HKEY_LOCAL_MACHINE\\SOFTWARE\\Microsoft\\Windows\\CurrentVersion\\Uninstall\\{{{msiProductId}}}", "UninstallString", null) as string) ??
 				(Registry.GetValue($"HKEY_LOCAL_MACHINE\\SOFTWARE\\WOW6432Node\\Microsoft\\Windows\\CurrentVersion\\Uninstall\\{{{msiProductId}}}", "UninstallString", null) as string);
@@ -972,7 +977,11 @@ namespace UnrealToolbox
 
 		async Task PollForUpdatesOnceAsync(CancellationToken cancellationToken)
 		{
-			using IHordeClientRef hordeClientRef = _hordeClientProvider.GetClientRef();
+			using IHordeClientRef? hordeClientRef = _hordeClientProvider.GetClientRef();
+			if (hordeClientRef == null)
+			{
+				return;
+			}
 
 			if (_serverUri != hordeClientRef.Client.ServerUrl)
 			{
