Unreal Toolbox: Set the app name for MacOS.

#rnx

diff --git a/Engine/Source/Programs/UnrealToolbox/App.axaml b/Engine/Source/Programs/UnrealToolbox/App.axaml
--- a/Engine/Source/Programs/UnrealToolbox/App.axaml
+++ b/Engine/Source/Programs/UnrealToolbox/App.axaml
@@ -2,6 +2,7 @@
 			xmlns:x="http://schemas.microsoft.com/winfx/2006/xaml"
 			xmlns:sty="using:FluentAvalonia.Styling"
 			x:Class="UnrealToolbox.App"
+			Name="Unreal Toolbox"
 			RequestedThemeVariant="Default">
 			<!-- "Default" ThemeVariant follows system theme variant. "Dark" or "Light" are other available options. -->
 	
