Unreal Toolbox: Disable functionality that requires named events on non-Windows platforms.

#rnx

diff --git a/Engine/Source/Programs/UnrealToolbox/App.axaml.cs b/Engine/Source/Programs/UnrealToolbox/App.axaml.cs
--- a/Engine/Source/Programs/UnrealToolbox/App.axaml.cs
+++ b/Engine/Source/Programs/UnrealToolbox/App.axaml.cs
@@ -138,6 +138,11 @@ namespace UnrealToolbox
 
 		private void WaitForEvents()
 		{
+			if (!OperatingSystem.IsWindows())
+			{
+				return; // Named global events are not currently supported on other platforms
+			}
+
 			using EventWaitHandle closeEvent = new EventWaitHandle(false, EventResetMode.AutoReset, Program.CloseEventName);
 			using EventWaitHandle refreshEvent = new EventWaitHandle(false, EventResetMode.AutoReset, Program.RefreshEventName);
 			using EventWaitHandle settingsEvent = new EventWaitHandle(false, EventResetMode.AutoReset, Program.SettingsEventName);
diff --git a/Engine/Source/Programs/UnrealToolbox/Program.cs b/Engine/Source/Programs/UnrealToolbox/Program.cs
--- a/Engine/Source/Programs/UnrealToolbox/Program.cs
+++ b/Engine/Source/Programs/UnrealToolbox/Program.cs
@@ -1,5 +1,6 @@
 // Copyright Epic Games, Inc. All Rights Reserved.
 
+using System.Runtime.Versioning;
 using Avalonia;
 using Avalonia.Controls;
 using EpicGames.Core;
@@ -45,24 +46,24 @@ namespace UnrealToolbox
 
 		static int RealMain(string[] args)
 		{
-			using EventWaitHandle closeEvent = new EventWaitHandle(false, EventResetMode.AutoReset, CloseEventName);
+			using EventWaitHandle? closeEvent = CreateNamedEvent(CloseEventName);
 			if (args.Any(x => x.Equals("-Close", StringComparison.OrdinalIgnoreCase)))
 			{
-				closeEvent.Set();
+				closeEvent?.Set();
 				return 0;
 			}
 
-			using EventWaitHandle refreshEvent = new EventWaitHandle(false, EventResetMode.AutoReset, RefreshEventName);
+			using EventWaitHandle? refreshEvent = CreateNamedEvent(RefreshEventName);
 			if (args.Any(x => x.Equals("-Refresh", StringComparison.OrdinalIgnoreCase)))
 			{
-				refreshEvent.Set();
+				refreshEvent?.Set();
 				return 0;
 			}
 
-			using EventWaitHandle settingsEvent = new EventWaitHandle(false, EventResetMode.AutoReset, SettingsEventName);
+			using EventWaitHandle? settingsEvent = CreateNamedEvent(SettingsEventName);
 			if (!args.Any(x => x.Equals("-Quiet", StringComparison.OrdinalIgnoreCase)))
 			{
-				settingsEvent.Set();
+				settingsEvent?.Set();
 			}
 
 			using EventWaitHandle commandEvent = new EventWaitHandle(false, EventResetMode.AutoReset, CommandEventName);
@@ -120,6 +121,18 @@ namespace UnrealToolbox
 				return DirectoryReference.Combine(appDataDir, "Epic Games", "Unreal Toolbox");
 			}
 		}
+
+		static EventWaitHandle? CreateNamedEvent(string name)
+		{
+			if (OperatingSystem.IsWindows())
+			{
+				return new EventWaitHandle(false, EventResetMode.AutoReset, RefreshEventName);
+			}
+			else
+			{
+				return null; // Need to figure out a way of replicating this functionality on MacOS/Linux.
+			}
+		}
 	}
 
 	class SingleInstanceMutex : IDisposable
