Horde: Set the default download folder to the users Downloads directory.

#rnx

diff --git a/Engine/Source/Programs/UnrealToolbox/Plugins/Artifacts/DownloadOptionsViewModel.cs b/Engine/Source/Programs/UnrealToolbox/Plugins/Artifacts/DownloadOptionsViewModel.cs
--- a/Engine/Source/Programs/UnrealToolbox/Plugins/Artifacts/DownloadOptionsViewModel.cs
+++ b/Engine/Source/Programs/UnrealToolbox/Plugins/Artifacts/DownloadOptionsViewModel.cs
@@ -145,11 +145,20 @@ namespace UnrealToolbox.Plugins.Artifacts
 		public DownloadOptionsViewModel(Window window, DownloadSettings? settings, IHordeClientProvider? hordeClientProvider)
 		{
 			_window = window;
-			_baseOutputDir = settings?.OutputDir ?? Directory.GetCurrentDirectory();
+			_baseOutputDir = settings?.OutputDir ?? 
+				GetDefaultDownloadFolder(window) ??
+				Environment.GetFolderPath(Environment.SpecialFolder.MyDocuments) ??
+				Directory.GetCurrentDirectory();
 			_appendBuildName = settings?.AppendBuildName ?? true;
 			_hordeClientProvider = hordeClientProvider;
 		}
 
+		static string? GetDefaultDownloadFolder(Window window)
+		{
+			IStorageFolder? storageFolder = Task.Run(async () => await window.StorageProvider.TryGetWellKnownFolderAsync(WellKnownFolder.Downloads)).Result;
+			return storageFolder?.TryGetLocalPath();
+		}
+
 		public DownloadOptions? GetDownloadOptions()
 		{
 			if (Artifact == null)
