Unreal Toolbox: Fix difficult-to-read hyperlink colors on about page in dark mode. Now uses standard hyperlink theming.

#rnx

diff --git a/Engine/Source/Programs/UnrealToolbox/AboutPage.axaml b/Engine/Source/Programs/UnrealToolbox/AboutPage.axaml
--- a/Engine/Source/Programs/UnrealToolbox/AboutPage.axaml
+++ b/Engine/Source/Programs/UnrealToolbox/AboutPage.axaml
@@ -13,10 +13,10 @@
 	</UserControl.Styles>
 
 	<ScrollViewer>
-		<StackPanel Margin="20" Spacing="5">
+		<StackPanel Margin="20" Spacing="25">
 			<TextBlock Theme="{StaticResource SubtitleTextBlockStyle}" VerticalAlignment="Center" Margin="0 0 0 0">About</TextBlock>
 
-			<Grid Margin="20">
+			<Grid>
 				<Grid.ColumnDefinitions>
 					<ColumnDefinition Width="Auto" MinWidth="180" />
 					<ColumnDefinition Width="*" />
@@ -28,18 +28,14 @@
 				</Grid.RowDefinitions>
 
 				<!-- Labels column -->
-				<TextBlock Grid.Row="0" Grid.Column="0" Text="Version:" VerticalAlignment="Center" Margin="0,0,10,10" />
-				<TextBlock Grid.Row="1" Grid.Column="0" Text="Settings and data directory:" VerticalAlignment="Center" Margin="0,0,10,10" />
-				<TextBlock Grid.Row="2" Grid.Column="0" Text="Log file:" VerticalAlignment="Center" Margin="0,0,10,10" />
+				<TextBlock Grid.Row="0" Grid.Column="0" Text="Version:" VerticalAlignment="Center" />
+				<TextBlock Grid.Row="1" Grid.Column="0" Text="Settings and data directory:" VerticalAlignment="Center" />
+				<TextBlock Grid.Row="2" Grid.Column="0" Text="Log file:" VerticalAlignment="Center" />
 
 				<!-- Values column -->
-				<TextBlock Grid.Row="0" Grid.Column="1" Name="_versionText" VerticalAlignment="Center" Margin="0,0,0,10" />
-				<Button Grid.Row="1" Grid.Column="1" Margin="0,0,0,10" Padding="0" Background="Transparent" BorderThickness="0" Name="_settingsDirButton" Cursor="Hand">
-					<TextBlock Name="_settingsDirText" Foreground="Blue" TextDecorations="Underline" />
-				</Button>
-				<Button Grid.Row="2" Grid.Column="1" Margin="0,0,0,10" Padding="0" Background="Transparent" BorderThickness="0" Name="_logFileButton" Cursor="Hand">
-					<TextBlock Name="_logFileText" Foreground="Blue" TextDecorations="Underline" />
-				</Button>
+				<TextBlock Grid.Row="0" Grid.Column="1" Name="_versionText" VerticalAlignment="Center" Padding="11,5,11,6" />
+				<HyperlinkButton Grid.Column="1" Grid.Row="1" FontSize="14" VerticalAlignment="Center" Name="_settingsDirButton"/>
+				<HyperlinkButton Grid.Column="1" Grid.Row="2" FontSize="14" VerticalAlignment="Center" Name="_logFileButton"/>
 			</Grid>
 		</StackPanel>
 	</ScrollViewer>
diff --git a/Engine/Source/Programs/UnrealToolbox/AboutPage.axaml.cs b/Engine/Source/Programs/UnrealToolbox/AboutPage.axaml.cs
--- a/Engine/Source/Programs/UnrealToolbox/AboutPage.axaml.cs
+++ b/Engine/Source/Programs/UnrealToolbox/AboutPage.axaml.cs
@@ -17,13 +17,13 @@ namespace UnrealToolbox
 			
 			string settingsDir = Program.DataDir.FullName;
 			_settingsDirButton.Click += (sender, args) => OpenFolder(settingsDir);
-			_settingsDirText.Text = settingsDir;
+			_settingsDirButton.Content = settingsDir;
 			
 			if (Program.LogFile != null)
 			{
 				string logFile = Program.LogFile.FullName;
 				_logFileButton.Click += (sender, args) => OpenFile(logFile);
-				_logFileText.Text = logFile;
+				_logFileButton.Content = logFile;
 			}
 		}
 		
