Unreal Toolbox: Bring the download window to the foreground when created.

#rnx

diff --git a/Engine/Source/Programs/UnrealToolbox/App.axaml.cs b/Engine/Source/Programs/UnrealToolbox/App.axaml.cs
--- a/Engine/Source/Programs/UnrealToolbox/App.axaml.cs
+++ b/Engine/Source/Programs/UnrealToolbox/App.axaml.cs
@@ -339,36 +339,15 @@ namespace UnrealToolbox
 			OpenSettings();
 		}
 
-		[DllImport("user32.dll")]
-		static extern uint SetForegroundWindow(nint hWnd);
-
 		private void OpenSettings()
 		{
-			if (_settingsWindow != null)
-			{
-				_settingsWindow.WindowState &= ~WindowState.Minimized;
-				_settingsWindow.BringIntoView();
-			}
-			else
+			if (_settingsWindow == null)
 			{
 				_settingsWindow = new SettingsWindow(_serviceProvider);
 				_settingsWindow.Closed += SettingsWindow_Closed;
 			}
 
-			_settingsWindow.Activate();
-			_settingsWindow.Show();
-
-			if (OperatingSystem.IsWindows())
-			{
-				// Avalonia doesn't seem to activate the window despite the call above, so fall back to pinvoke
-				nint? handle = _settingsWindow.TryGetPlatformHandle()?.Handle;
-				if (handle.HasValue)
-				{
-#pragma warning disable CA1806
-					SetForegroundWindow(handle.Value);
-#pragma warning restore CA1806
-				}
-			}
+			Program.BringWindowToFront(_settingsWindow);
 		}
 
 		private void SettingsWindow_Closed(object? sender, EventArgs e)
diff --git a/Engine/Source/Programs/UnrealToolbox/Plugins/Artifacts/ArtifactsPlugin.cs b/Engine/Source/Programs/UnrealToolbox/Plugins/Artifacts/ArtifactsPlugin.cs
--- a/Engine/Source/Programs/UnrealToolbox/Plugins/Artifacts/ArtifactsPlugin.cs
+++ b/Engine/Source/Programs/UnrealToolbox/Plugins/Artifacts/ArtifactsPlugin.cs
@@ -63,6 +63,8 @@ namespace UnrealToolbox.Plugins.Artifacts
 			DownloadOptionsWindow window = new DownloadOptionsWindow(artifactUrl, _settings, _hordeClientProvider);
 			window.Show();
 
+			Program.BringWindowToFront(window);
+
 			DownloadOptions? options = await window.Result;
 			if (options != null)
 			{
diff --git a/Engine/Source/Programs/UnrealToolbox/Program.cs b/Engine/Source/Programs/UnrealToolbox/Program.cs
--- a/Engine/Source/Programs/UnrealToolbox/Program.cs
+++ b/Engine/Source/Programs/UnrealToolbox/Program.cs
@@ -1,5 +1,6 @@
 // Copyright Epic Games, Inc. All Rights Reserved.
 
+using System.Runtime.InteropServices;
 using System.Runtime.Versioning;
 using Avalonia;
 using Avalonia.Controls;
@@ -133,6 +134,30 @@ namespace UnrealToolbox
 				return null; // Need to figure out a way of replicating this functionality on MacOS/Linux.
 			}
 		}
+
+		[DllImport("user32.dll")]
+		static extern uint SetForegroundWindow(nint hWnd);
+
+		public static void BringWindowToFront(Window window)
+		{
+			window.WindowState &= ~WindowState.Minimized;
+			window.BringIntoView();
+
+			window.Activate();
+			window.Show();
+
+			if (OperatingSystem.IsWindows())
+			{
+				// Avalonia doesn't seem to activate the window despite the call above, so fall back to pinvoke
+				nint? handle = window.TryGetPlatformHandle()?.Handle;
+				if (handle.HasValue)
+				{
+#pragma warning disable CA1806
+					SetForegroundWindow(handle.Value);
+#pragma warning restore CA1806
+				}
+			}
+		}
 	}
 
 	class SingleInstanceMutex : IDisposable
