Unreal Toolbox: Fix static analysis warnings.

#rnx

diff --git a/Engine/Source/Programs/UnrealToolbox/App.axaml.cs b/Engine/Source/Programs/UnrealToolbox/App.axaml.cs
--- a/Engine/Source/Programs/UnrealToolbox/App.axaml.cs
+++ b/Engine/Source/Programs/UnrealToolbox/App.axaml.cs
@@ -2,7 +2,6 @@
 
 using System.Diagnostics;
 using System.Reflection;
-using System.Runtime.InteropServices;
 using Avalonia;
 using Avalonia.Controls;
 using Avalonia.Controls.ApplicationLifetimes;
diff --git a/Engine/Source/Programs/UnrealToolbox/Program.cs b/Engine/Source/Programs/UnrealToolbox/Program.cs
--- a/Engine/Source/Programs/UnrealToolbox/Program.cs
+++ b/Engine/Source/Programs/UnrealToolbox/Program.cs
@@ -1,7 +1,6 @@
 // Copyright Epic Games, Inc. All Rights Reserved.
 
 using System.Runtime.InteropServices;
-using System.Runtime.Versioning;
 using Avalonia;
 using Avalonia.Controls;
 using EpicGames.Core;
@@ -127,7 +126,7 @@ namespace UnrealToolbox
 		{
 			if (OperatingSystem.IsWindows())
 			{
-				return new EventWaitHandle(false, EventResetMode.AutoReset, RefreshEventName);
+				return new EventWaitHandle(false, EventResetMode.AutoReset, name);
 			}
 			else
 			{
