Unreal Toolbox: Tweak font styles for captions.

#rnx

diff --git a/Engine/Source/Programs/UnrealToolbox/Styles.axaml b/Engine/Source/Programs/UnrealToolbox/Styles.axaml
--- a/Engine/Source/Programs/UnrealToolbox/Styles.axaml
+++ b/Engine/Source/Programs/UnrealToolbox/Styles.axaml
@@ -31,7 +31,7 @@
 		<Setter Property="VerticalAlignment" Value="Center"/>		
 	</Style>
 	<Style Selector="TextBlock.SettingHelp">
-		<Setter Property="Theme" Value="{DynamicResource BodyTextBlockStyle}"/>
+		<Setter Property="Theme" Value="{DynamicResource CaptionTextBlockStyle}"/>
 		<Setter Property="FontWeight" Value="Light"/>
 		<Setter Property="VerticalAlignment" Value="Center"/>
 	</Style>
