Unreal Toolbox: Support $(ToolDir) variables in UnrealToolbox commands.

#rnx

diff --git a/Engine/Source/Programs/UnrealToolbox/ToolCatalog.cs b/Engine/Source/Programs/UnrealToolbox/ToolCatalog.cs
--- a/Engine/Source/Programs/UnrealToolbox/ToolCatalog.cs
+++ b/Engine/Source/Programs/UnrealToolbox/ToolCatalog.cs
@@ -648,13 +648,17 @@ namespace UnrealToolbox
 
 		async Task<int> RunCommandAsync(string toolName, ToolCommand command, string workingDir, bool elevateProcess, bool hidden, CancellationToken cancellationToken)
 		{
+			Dictionary<string, string> properties = new Dictionary<string, string>(StringComparer.OrdinalIgnoreCase);
+			properties.Add("ToolDir", workingDir);
+
 			string arguments = String.Empty;
 			if (command.Arguments != null)
 			{
-				arguments = CommandLineArguments.Join(command.Arguments);
+				List<string> argumentList = command.Arguments.ConvertAll(x => StringUtils.ExpandProperties(x, properties));
+				arguments = CommandLineArguments.Join(argumentList);
 			}
 			
-			string fileName = Path.Combine(workingDir, command.FileName);
+			string fileName = Path.Combine(workingDir, StringUtils.ExpandProperties(command.FileName, properties));
 			if (!File.Exists(fileName))
 			{
 				fileName = command.FileName;
