UnrealToolbox: Don't allow updates to start while any windows are open.

#rnx

diff --git a/Engine/Source/Programs/UnrealToolbox/App.axaml.cs b/Engine/Source/Programs/UnrealToolbox/App.axaml.cs
--- a/Engine/Source/Programs/UnrealToolbox/App.axaml.cs
+++ b/Engine/Source/Programs/UnrealToolbox/App.axaml.cs
@@ -29,12 +29,12 @@ namespace UnrealToolbox
 		WindowIcon? _busyIcon;
 		WindowIcon? _pausedIcon;
 		WindowIcon? _errorIcon;
+		DispatcherTimer? _retryUpdateTimer;
 
 		SettingsWindow? _settingsWindow;
 
 		Thread? _settingsThread;
 		ManualResetEvent? _settingsThreadStop;
-		bool _shutdownOnClose;
 
 		/// <summary>
 		/// Constructor
@@ -215,15 +215,28 @@ namespace UnrealToolbox
 
 		void UpdateReadyMainThread()
 		{
-			if (_settingsWindow == null)
+			if (AllWindowsClosed())
 			{
 				_logger.LogInformation("Shutting down to install update");
 				((IClassicDesktopStyleApplicationLifetime)ApplicationLifetime!).Shutdown();
 			}
-			else
+			else if (_retryUpdateTimer == null)
 			{
 				_logger.LogInformation("Scheduling update when settings window closes");
-				_shutdownOnClose = true;
+				_retryUpdateTimer = new DispatcherTimer(TimeSpan.FromSeconds(30.0), DispatcherPriority.Default, (_, _) => UpdateReadyMainThread());
+			}
+		}
+
+		bool AllWindowsClosed()
+		{
+			IClassicDesktopStyleApplicationLifetime? lifetime = ApplicationLifetime as IClassicDesktopStyleApplicationLifetime;
+			if (lifetime != null && lifetime.Windows.Count == 0)
+			{
+				return true;
+			}
+			else
+			{
+				return false;
 			}
 		}
 
@@ -352,12 +365,6 @@ namespace UnrealToolbox
 		private void SettingsWindow_Closed(object? sender, EventArgs e)
 		{
 			_settingsWindow = null;
-
-			if (_shutdownOnClose)
-			{
-				_logger.LogInformation("Running update due to settings window closing");
-				((IClassicDesktopStyleApplicationLifetime)ApplicationLifetime!).Shutdown();
-			}
 		}
 
 		private void TrayIcon_Exit(object? sender, EventArgs e)
