UnrealToolbox: Add support for workspace patching.

#rnx

diff --git a/Engine/Source/Programs/UnrealToolbox/Plugins/Artifacts/ArtifactsPlugin.cs b/Engine/Source/Programs/UnrealToolbox/Plugins/Artifacts/ArtifactsPlugin.cs
--- a/Engine/Source/Programs/UnrealToolbox/Plugins/Artifacts/ArtifactsPlugin.cs
+++ b/Engine/Source/Programs/UnrealToolbox/Plugins/Artifacts/ArtifactsPlugin.cs
@@ -88,9 +88,6 @@ namespace UnrealToolbox.Plugins.Artifacts
 			{
 				CancellationToken linkedCancellationToken = cancellationTokenSource.Token;
 
-				DirectoryReference stateFolder = DirectoryReference.Combine(options.OutputDir, ".ugs");
-				DirectoryReference.CreateDirectory(stateFolder);
-
 				Uri serverUrl = new Uri(options.BaseUrl.GetLeftPart(UriPartial.Authority));
 
 				using IHordeClientRef? hordeClientRef = _hordeClientProvider.GetClientRef();
@@ -103,11 +100,24 @@ namespace UnrealToolbox.Plugins.Artifacts
 
 					IBlobRef<DirectoryNode>? blobRef = await storageNamespace.ReadRefAsync<DirectoryNode>(options.RefName, cancellationToken: linkedCancellationToken);
 
-					DirectoryNode node = await blobRef.ReadBlobAsync(cancellationToken: linkedCancellationToken);
-					progress.SetTotalSize(node.Length);
+					if (options.PatchExistingData)
+					{
+						Workspace workspace = await Workspace.CreateOrOpenAsync(options.OutputDir, _logger, cancellationToken);
+
+						progress.Status = "Reconciling current data...";
+						await workspace.ReconcileAsync(cancellationToken);
+
+						DirectoryNode node = await blobRef.ReadBlobAsync(cancellationToken: linkedCancellationToken);
+						progress.SetTotalSize(node.Length);
 
-					progress.Status = "Starting...";
-					await blobRef.ExtractAsync(options.OutputDir.ToDirectoryInfo(), new ExtractOptions { Progress = progress, ProgressUpdateFrequency = TimeSpan.FromSeconds(0.2) }, _logger, linkedCancellationToken);
+						progress.Status = "Starting...";
+						await workspace.SyncAsync(WorkspaceLayerId.Default, blobRef, new ExtractOptions { Progress = progress, ProgressUpdateFrequency = TimeSpan.FromSeconds(0.2) }, cancellationToken: cancellationToken);
+					}
+					else
+					{
+						progress.Status = "Starting...";
+						await blobRef.ExtractAsync(options.OutputDir.ToDirectoryInfo(), new ExtractOptions { Progress = progress, ProgressUpdateFrequency = TimeSpan.FromSeconds(0.2) }, _logger, linkedCancellationToken);
+					}
 
 					progress.Status = "Download complete.";
 					progress.MarkComplete();
diff --git a/Engine/Source/Programs/UnrealToolbox/Plugins/Artifacts/DownloadOptions.cs b/Engine/Source/Programs/UnrealToolbox/Plugins/Artifacts/DownloadOptions.cs
--- a/Engine/Source/Programs/UnrealToolbox/Plugins/Artifacts/DownloadOptions.cs
+++ b/Engine/Source/Programs/UnrealToolbox/Plugins/Artifacts/DownloadOptions.cs
@@ -10,12 +10,14 @@ namespace UnrealToolbox.Plugins.Artifacts
 		public Uri BaseUrl { get; set; }
 		public RefName RefName { get; set; }
 		public DirectoryReference OutputDir { get; set; }
+		public bool PatchExistingData { get; set; }
 
-		public DownloadOptions(Uri baseUrl, RefName refName, DirectoryReference outputDir)
+		public DownloadOptions(Uri baseUrl, RefName refName, DirectoryReference outputDir, bool patchExistingData)
 		{
 			BaseUrl = baseUrl;
 			RefName = refName;
 			OutputDir = outputDir;
+			PatchExistingData = patchExistingData;
 		}
 	}
 }
diff --git a/Engine/Source/Programs/UnrealToolbox/Plugins/Artifacts/DownloadOptionsViewModel.cs b/Engine/Source/Programs/UnrealToolbox/Plugins/Artifacts/DownloadOptionsViewModel.cs
--- a/Engine/Source/Programs/UnrealToolbox/Plugins/Artifacts/DownloadOptionsViewModel.cs
+++ b/Engine/Source/Programs/UnrealToolbox/Plugins/Artifacts/DownloadOptionsViewModel.cs
@@ -78,6 +78,9 @@ namespace UnrealToolbox.Plugins.Artifacts
 		[NotifyPropertyChangedFor(nameof(OutputDir))]
 		public bool _appendBuildName;
 
+		[ObservableProperty]
+		public bool _patchExistingData;
+
 		[ObservableProperty]
 		[NotifyPropertyChangedFor(nameof(ShowErrorMessage))]
 		string _errorMessage = String.Empty;
@@ -150,6 +153,7 @@ namespace UnrealToolbox.Plugins.Artifacts
 				Environment.GetFolderPath(Environment.SpecialFolder.MyDocuments) ??
 				Directory.GetCurrentDirectory();
 			_appendBuildName = settings?.AppendBuildName ?? true;
+			_patchExistingData = settings?.PatchExistingData ?? false;
 			_hordeClientProvider = hordeClientProvider;
 		}
 
@@ -170,7 +174,7 @@ namespace UnrealToolbox.Plugins.Artifacts
 				return null;
 			}
 
-			return new DownloadOptions(Artifact.BaseUrl, Artifact.RefName, new DirectoryReference(OutputDir));
+			return new DownloadOptions(Artifact.BaseUrl, Artifact.RefName, new DirectoryReference(OutputDir), PatchExistingData);
 		}
 
 		public async Task<bool> RefreshArtifactAsync()
diff --git a/Engine/Source/Programs/UnrealToolbox/Plugins/Artifacts/DownloadOptionsWindow.axaml b/Engine/Source/Programs/UnrealToolbox/Plugins/Artifacts/DownloadOptionsWindow.axaml
--- a/Engine/Source/Programs/UnrealToolbox/Plugins/Artifacts/DownloadOptionsWindow.axaml
+++ b/Engine/Source/Programs/UnrealToolbox/Plugins/Artifacts/DownloadOptionsWindow.axaml
@@ -61,6 +61,7 @@
 					<Button Margin="16 0 0 0" Grid.Column="1" Command="{Binding SelectOutputDirAsync}" Width="120">Browse...</Button>
 				</Grid>
 				<CheckBox IsChecked="{Binding AppendBuildName}" Content="Append build name"/>
+				<CheckBox IsChecked="{Binding PatchExistingData}" Content="Reuse existing data for patching"/>
 			</StackPanel>
 		</Border>
 
diff --git a/Engine/Source/Programs/UnrealToolbox/Plugins/Artifacts/DownloadOptionsWindow.axaml.cs b/Engine/Source/Programs/UnrealToolbox/Plugins/Artifacts/DownloadOptionsWindow.axaml.cs
--- a/Engine/Source/Programs/UnrealToolbox/Plugins/Artifacts/DownloadOptionsWindow.axaml.cs
+++ b/Engine/Source/Programs/UnrealToolbox/Plugins/Artifacts/DownloadOptionsWindow.axaml.cs
@@ -54,6 +54,7 @@ namespace UnrealToolbox.Plugins.Artifacts
 					DownloadSettings newSettings = new DownloadSettings(_settings.Current);
 					newSettings.OutputDir = _viewModel.BaseOutputDir;
 					newSettings.AppendBuildName = _viewModel.AppendBuildName;
+					newSettings.PatchExistingData = _viewModel.PatchExistingData;
 					_settings.UpdateSettings(newSettings);
 				}
 
diff --git a/Engine/Source/Programs/UnrealToolbox/Plugins/Artifacts/DownloadSettings.cs b/Engine/Source/Programs/UnrealToolbox/Plugins/Artifacts/DownloadSettings.cs
--- a/Engine/Source/Programs/UnrealToolbox/Plugins/Artifacts/DownloadSettings.cs
+++ b/Engine/Source/Programs/UnrealToolbox/Plugins/Artifacts/DownloadSettings.cs
@@ -18,6 +18,11 @@ namespace UnrealToolbox.Plugins.Artifacts
 		/// </summary>
 		public bool AppendBuildName { get; set; }
 
+		/// <summary>
+		/// Whether to patch existing data when writing the new build
+		/// </summary>
+		public bool PatchExistingData { get; set; }
+
 		/// <summary>
 		/// Default constructor
 		/// </summary>
@@ -31,6 +36,7 @@ namespace UnrealToolbox.Plugins.Artifacts
 		{
 			OutputDir = other.OutputDir;
 			AppendBuildName = other.AppendBuildName;
+			PatchExistingData = other.PatchExistingData;
 		}
 	}
 }
